Move GeneSplicer into public plugins folder

diff --git a/Engine/Plugins/Animation/GeneSplicer/GeneSplicer.uplugin b/Engine/Plugins/Animation/GeneSplicer/GeneSplicer.uplugin
new file mode 100644
--- /dev/null
+++ b/Engine/Plugins/Animation/GeneSplicer/GeneSplicer.uplugin
@@ -0,0 +1,47 @@
+{
+	"FileVersion": 3,
+	"Version": 9.8,
+	"VersionName": "9.8.0",
+	"FriendlyName": "GeneSplicer Plugin v9.8.0",
+	"Description": "3Lateral GeneSplicer plugin for facial animation",
+	"Category": "Animation",
+	"CreatedBy": "3Lateral",
+	"CreatedByURL": "http://www.3lateral.com/",
+	"DocsURL": "",
+	"MarketplaceURL": "",
+	"SupportURL": "",
+	"CanContainContent": false,
+	"IsBetaVersion": true,
+	"Installed": true,
+	"Modules": [
+		{
+			"Name": "GeneSplicerLib",
+			"Type": "Runtime",
+			"LoadingPhase": "PreDefault"
+		},
+		{
+			"Name": "GeneSplicerLibTest",
+			"Type": "Runtime",
+			"LoadingPhase": "PreDefault",
+			"PlatformAllowList": [
+				"Win64"
+			]
+		},
+		{
+			"Name": "GeneSplicerModule",
+			"Type": "Runtime",
+			"LoadingPhase": "Default"
+		}
+	],
+	"Plugins": [
+		{
+			"Name": "RigLogic",
+			"Enabled": true
+		},
+		{		
+			"Name": "ControlRig",
+			"Enabled": true
+		}
+
+	]
+}
\ No newline at end of file
diff --git a/Engine/Plugins/Animation/GeneSplicer/Source/GeneSplicerLib/Private/genesplicer/GeneSplicer.cpp b/Engine/Plugins/Animation/GeneSplicer/Source/GeneSplicerLib/Private/genesplicer/GeneSplicer.cpp
new file mode 100644
--- /dev/null
+++ b/Engine/Plugins/Animation/GeneSplicer/Source/GeneSplicerLib/Private/genesplicer/GeneSplicer.cpp
@@ -0,0 +1,55 @@
+// Copyright Epic Games, Inc. All Rights Reserved.
+
+#include "genesplicer/GeneSplicer.h"
+
+#include "genesplicer/CalculationType.h"
+#include "genesplicer/GeneSplicerImpl.h"
+#include "genesplicer/splicedata/GenePool.h"
+
+namespace gs4 {
+
+GeneSplicer::GeneSplicer(CalculationType calculationType, MemoryResource* memRes) {
+    switch (calculationType) {
+        case CalculationType::SSE:
+            pImpl = ScopedPtr<Impl, FactoryDestroy<Impl> >(Impl::create<CalculationType::SSE>(memRes));
+            break;
+        case CalculationType::AVX:
+            pImpl = ScopedPtr<Impl, FactoryDestroy<Impl> >(Impl::create<CalculationType::AVX>(memRes));
+            break;
+        case CalculationType::Scalar:
+        default:
+            pImpl = ScopedPtr<Impl, FactoryDestroy<Impl> >(Impl::create<CalculationType::Scalar>(memRes));
+            break;
+    }
+}
+
+GeneSplicer::~GeneSplicer() = default;
+
+GeneSplicer::GeneSplicer(GeneSplicer&& rhs) = default;
+GeneSplicer& GeneSplicer::operator=(GeneSplicer&& rhs) = default;
+
+void GeneSplicer::splice(const SpliceData* spliceData, GeneSplicerDNAReader* master) {
+    pImpl->splice(*spliceData, master);
+}
+
+void GeneSplicer::spliceNeutralMeshes(const SpliceData* spliceData, GeneSplicerDNAReader* master) {
+    pImpl->spliceNeutralMeshes(*spliceData, master);
+}
+
+void GeneSplicer::spliceBlendShapes(const SpliceData* spliceData, GeneSplicerDNAReader* master) {
+    pImpl->spliceBlendShapes(*spliceData, master);
+}
+
+void GeneSplicer::spliceNeutralJoints(const SpliceData* spliceData, GeneSplicerDNAReader* master) {
+    pImpl->spliceNeutralJoints(*spliceData, master);
+}
+
+void GeneSplicer::spliceJointBehavior(const SpliceData* spliceData, GeneSplicerDNAReader* master) {
+    pImpl->spliceJointBehavior(*spliceData, master);
+}
+
+void GeneSplicer::spliceSkinWeights(const SpliceData* spliceData, GeneSplicerDNAReader* master) {
+    pImpl->spliceSkinWeights(*spliceData, master);
+}
+
+}  // namespace gs4
diff --git a/Engine/Plugins/Animation/GeneSplicer/Source/GeneSplicerLib/Private/genesplicer/GeneSplicerImpl.cpp b/Engine/Plugins/Animation/GeneSplicer/Source/GeneSplicerLib/Private/genesplicer/GeneSplicerImpl.cpp
new file mode 100644
--- /dev/null
+++ b/Engine/Plugins/Animation/GeneSplicer/Source/GeneSplicerLib/Private/genesplicer/GeneSplicerImpl.cpp
@@ -0,0 +1,59 @@
+// Copyright Epic Games, Inc. All Rights Reserved.
+
+#include "genesplicer/GeneSplicerImpl.h"
+
+#include "genesplicer/Macros.h"
+#include "genesplicer/Splicer.h"
+#include "genesplicer/splicedata/SpliceData.h"
+#include "genesplicer/types/PImplExtractor.h"
+
+namespace gs4 {
+
+GeneSplicer::Impl::Impl(SplicerPtr neutralJointSplicer_,
+                        SplicerPtr jointBehaviorSplicer_,
+                        SplicerPtr blendShapeSplicer_,
+                        SplicerPtr neutralMeshSplicer_,
+                        SplicerPtr skinWeightSplicer_,
+                        MemoryResource* memRes_) :
+    memRes{memRes_},
+    neutralJointSplicer{std::move(neutralJointSplicer_)},
+    jointBehaviorSplicer{std::move(jointBehaviorSplicer_)},
+    blendShapeSplicer{std::move(blendShapeSplicer_)},
+    neutralMeshSplicer{std::move(neutralMeshSplicer_)},
+    skinWeightSplicer{std::move(skinWeightSplicer_)} {
+}
+
+void GeneSplicer::Impl::destroy(GeneSplicer::Impl* instance) {
+    PolyAllocator<GeneSplicer::Impl> alloc{instance->memRes};
+    alloc.deleteObject(instance);
+}
+
+void GeneSplicer::Impl::splice(const SpliceData& spliceData, GeneSplicerDNAReader* master) {
+    spliceNeutralMeshes(spliceData, master);
+    spliceBlendShapes(spliceData, master);
+    spliceNeutralJoints(spliceData, master);
+    spliceJointBehavior(spliceData, master);
+    spliceSkinWeights(spliceData, master);
+}
+
+void GeneSplicer::Impl::spliceNeutralMeshes(const SpliceData& spliceData, GeneSplicerDNAReader* master) {
+    neutralMeshSplicer.get()->splice(PImplExtractor<SpliceData>::get(spliceData), master);
+}
+
+void GeneSplicer::Impl::spliceBlendShapes(const SpliceData& spliceData, GeneSplicerDNAReader* master) {
+    blendShapeSplicer.get()->splice(PImplExtractor<SpliceData>::get(spliceData), master);
+}
+
+void GeneSplicer::Impl::spliceNeutralJoints(const SpliceData& spliceData, GeneSplicerDNAReader* master) {
+    neutralJointSplicer.get()->splice(PImplExtractor<SpliceData>::get(spliceData), master);
+}
+
+void GeneSplicer::Impl::spliceJointBehavior(const SpliceData& spliceData, GeneSplicerDNAReader* master) {
+    jointBehaviorSplicer.get()->splice(PImplExtractor<SpliceData>::get(spliceData), master);
+}
+
+void GeneSplicer::Impl::spliceSkinWeights(const SpliceData& spliceData, GeneSplicerDNAReader* master) {
+    skinWeightSplicer.get()->splice(PImplExtractor<SpliceData>::get(spliceData), master);
+}
+
+}  // namespace gs4
diff --git a/Engine/Plugins/Animation/GeneSplicer/Source/GeneSplicerLib/Private/genesplicer/GeneSplicerImpl.h b/Engine/Plugins/Animation/GeneSplicer/Source/GeneSplicerLib/Private/genesplicer/GeneSplicerImpl.h
new file mode 100644
--- /dev/null
+++ b/Engine/Plugins/Animation/GeneSplicer/Source/GeneSplicerLib/Private/genesplicer/GeneSplicerImpl.h
@@ -0,0 +1,66 @@
+// Copyright Epic Games, Inc. All Rights Reserved.
+
+#pragma once
+
+#include "genesplicer/GeneSplicer.h"
+#include "genesplicer/Splicer.h"
+#include "genesplicer/TypeDefs.h"
+#include "genesplicer/blendshapesplicer/BlendShapeSplicer.h"
+#include "genesplicer/geometrysplicer/NeutralMeshSplicer.h"
+#include "genesplicer/jointbehaviorsplicer/JointBehaviorSplicer.h"
+#include "genesplicer/neutraljointsplicer/NeutralJointSplicer.h"
+#include "genesplicer/skinweightsplicer/SkinWeightSplicer.h"
+
+namespace gs4 {
+
+class GeneSplicerDNAReader;
+class SpliceData;
+
+class GeneSplicer::Impl {
+    private:
+        using SplicerPtr = pma::UniqueInstance<Splicer>::PointerType;
+
+    public:
+        Impl(SplicerPtr neutralJointSplicer_,
+             SplicerPtr jointBehaviorSplicer_,
+             SplicerPtr blendShapeSplicer_,
+             SplicerPtr neutralMeshSplicer_,
+             SplicerPtr skinWeightSplicer_,
+             MemoryResource* memRes);
+
+        template<CalculationType CT>
+        static Impl* create(MemoryResource* memRes) {
+            PolyAllocator<GeneSplicer::Impl> alloc{memRes};
+            SplicerPtr neutralJointSplicer =
+                pma::UniqueInstance<NeutralJointSplicer<CT>, Splicer>::with(memRes).create(memRes);
+            SplicerPtr jointBehaviorSplicer = pma::UniqueInstance<JointBehaviorSplicer<CT>, Splicer>::with(memRes).create(memRes);
+            SplicerPtr blendShapeSplicer =
+                pma::UniqueInstance<BlendShapeSplicer<CT>, Splicer>::with(memRes).create(memRes);
+            SplicerPtr neutralMeshSplicer = pma::UniqueInstance<NeutralMeshSplicer<CT>, Splicer>::with(memRes).create(memRes);
+            SplicerPtr skinWeightSplicer = pma::UniqueInstance<SkinWeightSplicer<CT>, Splicer>::with(memRes).create(memRes);
+            return alloc.newObject(std::move(neutralJointSplicer),
+                                   std::move(jointBehaviorSplicer),
+                                   std::move(blendShapeSplicer),
+                                   std::move(neutralMeshSplicer),
+                                   std::move(skinWeightSplicer),
+                                   memRes);
+        }
+
+        static void destroy(Impl* instance);
+        void splice(const SpliceData& spliceData, GeneSplicerDNAReader* master);
+        void spliceNeutralJoints(const SpliceData& spliceData, GeneSplicerDNAReader* master);
+        void spliceNeutralMeshes(const SpliceData& spliceData, GeneSplicerDNAReader* master);
+        void spliceBlendShapes(const SpliceData& spliceData, GeneSplicerDNAReader* master);
+        void spliceJointBehavior(const SpliceData& spliceData, GeneSplicerDNAReader* master);
+        void spliceSkinWeights(const SpliceData& spliceData, GeneSplicerDNAReader* master);
+
+    private:
+        MemoryResource* memRes;
+        SplicerPtr neutralJointSplicer;
+        SplicerPtr jointBehaviorSplicer;
+        SplicerPtr blendShapeSplicer;
+        SplicerPtr neutralMeshSplicer;
+        SplicerPtr skinWeightSplicer;
+};
+
+}  // namespace gs4
diff --git a/Engine/Plugins/Animation/GeneSplicer/Source/GeneSplicerLib/Private/genesplicer/geometrysplicer/NeutralMeshSplicer.cpp b/Engine/Plugins/Animation/GeneSplicer/Source/GeneSplicerLib/Private/genesplicer/geometrysplicer/NeutralMeshSplicer.cpp
new file mode 100644
--- /dev/null
+++ b/Engine/Plugins/Animation/GeneSplicer/Source/GeneSplicerLib/Private/genesplicer/geometrysplicer/NeutralMeshSplicer.cpp
@@ -0,0 +1,58 @@
+// Copyright Epic Games, Inc. All Rights Reserved.
+
+#include "genesplicer/geometrysplicer/NeutralMeshSplicer.h"
+
+#include "genesplicer/GeneSplicerDNAReaderImpl.h"
+#include "genesplicer/Macros.h"
+#include "genesplicer/TypeDefs.h"
+#include "genesplicer/splicedata/PoolSpliceParamsFilter.h"
+#include "genesplicer/splicedata/PoolSpliceParamsImpl.h"
+#include "genesplicer/splicedata/SpliceDataImpl.h"
+#include "genesplicer/splicedata/genepool/GenePoolImpl.h"
+#include "genesplicer/types/BlockStorage.h"
+#include "genesplicer/types/Vec3.h"
+
+#include <cstdint>
+
+namespace gs4 {
+
+template<CalculationType CT>
+void NeutralMeshSplicer<CT>::splice(const SpliceDataInterface* spliceData, GeneSplicerDNAReader* master_) {
+    auto master = static_cast<GeneSplicerDNAReaderImpl*>(master_);
+    auto masterMemRes = master->getMemoryResource();
+
+    const auto& baseArch = spliceData->getBaseArchetype();
+    const auto& baseArchMeshes = baseArch.getNeutralMeshes();
+
+    auto neutralMeshPredicate = [](const RawGenes& baseArchetype, const PoolSpliceParamsImpl* pool, std::uint16_t meshIndex) {
+            return baseArchetype.getNeutralMeshes()[meshIndex].size() == pool->getGenePool()->getVertexPositionCount(meshIndex);
+        };
+    auto poolsToSplicePerMesh = filterPoolSpliceParamsPerMesh(spliceData, neutralMeshPredicate, masterMemRes);
+
+    for (std::uint16_t meshIdx = 0u; meshIdx <  poolsToSplicePerMesh.rowCount(); meshIdx++) {
+
+        const auto& baseArchMesh = baseArchMeshes[meshIdx];
+        RawVector3Vector resultingVertices = constructWithPadding(baseArchMesh, masterMemRes);
+
+        for (auto poolParams : poolsToSplicePerMesh[meshIdx]) {
+            auto genePool = poolParams->getGenePool();
+            const auto& vertexWeights = poolParams->getVertexWeightsData();
+            const auto& neutralMeshes = genePool->getNeutralMeshes();
+            float scale = poolParams->getScale();
+
+            BlockSplicer<CT>::splice(Matrix2DView<const XYZBlock<16u> >{neutralMeshes[meshIdx]},
+                                     Matrix2DView<const VBlock<16u> > {vertexWeights[meshIdx]},
+                                     poolParams->getDNAIndices(),
+                                     resultingVertices,
+                                     scale);
+        }
+        resultingVertices.resize(master->getVertexPositionCount(meshIdx));
+        master->setVertexPositions(meshIdx, std::move(resultingVertices));
+    }
+}
+
+template class NeutralMeshSplicer<CalculationType::Scalar>;
+template class NeutralMeshSplicer<CalculationType::SSE>;
+template class NeutralMeshSplicer<CalculationType::AVX>;
+
+}  // namespace gs4
diff --git a/Engine/Plugins/Animation/GeneSplicer/Source/GeneSplicerLib/Private/genesplicer/neutraljointsplicer/JointAttributeSpecialization.h b/Engine/Plugins/Animation/GeneSplicer/Source/GeneSplicerLib/Private/genesplicer/neutraljointsplicer/JointAttributeSpecialization.h
new file mode 100644
--- /dev/null
+++ b/Engine/Plugins/Animation/GeneSplicer/Source/GeneSplicerLib/Private/genesplicer/neutraljointsplicer/JointAttributeSpecialization.h
@@ -0,0 +1,54 @@
+// Copyright Epic Games, Inc. All Rights Reserved.
+
+#pragma once
+
+#include "genesplicer/GeneSplicerDNAReaderImpl.h"
+#include "genesplicer/Macros.h"
+#include "genesplicer/neutraljointsplicer/JointAttribute.h"
+#include "genesplicer/types/Vec3.h"
+
+namespace gs4 {
+
+template<JointAttribute T>
+FORCE_INLINE ConstVec3VectorView getJointAttributeValues(const Reader* const reader);
+
+template<>
+FORCE_INLINE ConstVec3VectorView getJointAttributeValues<JointAttribute::Translation>(const Reader* const reader) {
+    return ConstVec3VectorView{reader->getNeutralJointTranslationXs(), reader->getNeutralJointTranslationYs(),
+                                                                       reader->getNeutralJointTranslationZs()};
+}
+
+template<>
+FORCE_INLINE ConstVec3VectorView getJointAttributeValues<JointAttribute::Rotation>(const Reader* const reader) {
+    return ConstVec3VectorView{reader->getNeutralJointRotationXs(), reader->getNeutralJointRotationYs(),
+                                                                    reader->getNeutralJointRotationZs()};
+}
+
+template<JointAttribute T>
+FORCE_INLINE Vec3 getJointAttributeValue(const Reader* const reader, std::uint16_t jointIndex);
+
+template<>
+FORCE_INLINE Vec3 getJointAttributeValue<JointAttribute::Translation>(const Reader* const reader, std::uint16_t jointIndex) {
+    return reader->getNeutralJointTranslation(jointIndex);
+}
+
+template<>
+FORCE_INLINE Vec3 getJointAttributeValue<JointAttribute::Rotation>(const Reader* const reader, std::uint16_t jointIndex) {
+    return reader->getNeutralJointRotation(jointIndex);
+}
+
+template<JointAttribute T>
+FORCE_INLINE void setJointAttributeValues(GeneSplicerDNAReaderImpl* master, RawVector3Vector& result);
+
+template<>
+FORCE_INLINE void setJointAttributeValues<JointAttribute::Translation>(GeneSplicerDNAReaderImpl* master,
+                                                                       RawVector3Vector& result) {
+    master->setNeutralJointTranslations(std::move(result));
+}
+
+template<>
+FORCE_INLINE void setJointAttributeValues<JointAttribute::Rotation>(GeneSplicerDNAReaderImpl* master, RawVector3Vector& result) {
+    master->setNeutralJointRotations(std::move(result));
+}
+
+}  // namespace gs4
diff --git a/Engine/Plugins/Animation/GeneSplicer/Source/GeneSplicerLib/Private/genesplicer/neutraljointsplicer/NeutralJointSplicer.cpp b/Engine/Plugins/Animation/GeneSplicer/Source/GeneSplicerLib/Private/genesplicer/neutraljointsplicer/NeutralJointSplicer.cpp
new file mode 100644
--- /dev/null
+++ b/Engine/Plugins/Animation/GeneSplicer/Source/GeneSplicerLib/Private/genesplicer/neutraljointsplicer/NeutralJointSplicer.cpp
@@ -0,0 +1,87 @@
+// Copyright Epic Games, Inc. All Rights Reserved.
+
+#include "genesplicer/neutraljointsplicer/NeutralJointSplicer.h"
+
+#include "genesplicer/GeneSplicerDNAReaderImpl.h"
+#include "genesplicer/Macros.h"
+#include "genesplicer/splicedata/PoolSpliceParams.h"
+#include "genesplicer/splicedata/PoolSpliceParamsImpl.h"
+#include "genesplicer/splicedata/SpliceDataImpl.h"
+#include "genesplicer/splicedata/genepool/GenePoolImpl.h"
+#include "genesplicer/splicedata/genepool/RawNeutralJoints.h"
+#include "genesplicer/neutraljointsplicer/JointAttribute.h"
+#include "genesplicer/neutraljointsplicer/JointAttributeSpecialization.h"
+#include "genesplicer/types/Aliases.h"
+#include "genesplicer/types/Matrix.h"
+#include "genesplicer/types/PImplExtractor.h"
+
+#ifdef _MSC_VER
+    #pragma warning(push)
+    #pragma warning(disable : 4365 4987)
+#endif
+#include <utility>
+#ifdef _MSC_VER
+    #pragma warning(pop)
+#endif
+
+namespace gs4 {
+
+template<CalculationType CT>
+void NeutralJointSplicer<CT>::splice(const SpliceDataInterface* spliceData, GeneSplicerDNAReader* master_) {
+    auto master = static_cast<GeneSplicerDNAReaderImpl*>(master_);
+    auto splice =
+        [spliceData, master](JointAttribute jointAttribute,
+                             std::function<float(const PoolSpliceParamsImpl* poolParams)> scaleGetter) {
+            auto masterMemRes = master->getMemoryResource();
+            const auto& baseArch = spliceData->getBaseArchetype();
+            auto poolParams = spliceData->getAllPoolParams();
+            const auto& baseArchJoints = baseArch.getNeutralJoints(jointAttribute);
+            RawVector3Vector resultingJoints = constructWithPadding(baseArchJoints, masterMemRes);
+
+            for (auto pool : poolParams) {
+                auto genePool = pool->getGenePool();
+                if (genePool->getNeutralJointCount() != baseArch.getJointCount()) {
+                    continue;
+                }
+                const auto& neutralJoints = genePool->getNeutralJoints(jointAttribute);
+                const auto& jointWeights = pool->getJointWeightsData();
+                auto dnaIndices = pool->getDNAIndices();
+                BlockSplicer<CT>::splice(Matrix2DView<const XYZBlock<16u> >{neutralJoints},
+                                         Matrix2DView<const VBlock<16u> >{jointWeights},
+                                         dnaIndices,
+                                         resultingJoints,
+                                         scaleGetter(pool));
+            }
+            resultingJoints.resize(baseArch.getJointCount());
+            return resultingJoints;
+        };
+    RawNeutralJoints neutralJoints{master->getMemoryResource()};
+    neutralJoints.translations = splice(JointAttribute::Translation, [](const PoolSpliceParamsImpl* poolParams) {
+            return poolParams->getScale();
+        });
+    // we want to avoid scaling of rotations, so we are pegging it to 1.0f
+    neutralJoints.rotations = splice(JointAttribute::Rotation, [](const PoolSpliceParamsImpl*  /*unused*/) {
+            return 1.0f;
+        });
+    auto getJointParentIndex = std::bind(&dna::DefinitionReader::getJointParentIndex, master_, std::placeholders::_1);
+    // Since GenePools and RawGenes of archetype hold joints in world space we need to convert them to local space
+    toLocalSpace(getJointParentIndex, neutralJoints);
+
+    auto toDegrees = [](float& angle) {
+            angle = tdm::degrees(angle);
+        };
+
+    auto& rotations = neutralJoints.rotations;
+    std::for_each(rotations.xs.begin(), rotations.xs.end(), toDegrees);
+    std::for_each(rotations.ys.begin(), rotations.ys.end(), toDegrees);
+    std::for_each(rotations.zs.begin(), rotations.zs.end(), toDegrees);
+
+    master->setNeutralJointTranslations(std::move(neutralJoints.translations));
+    master->setNeutralJointRotations(std::move(neutralJoints.rotations));
+}
+
+template class NeutralJointSplicer<CalculationType::Scalar>;
+template class NeutralJointSplicer<CalculationType::SSE>;
+template class NeutralJointSplicer<CalculationType::AVX>;
+
+}  // namespace gs4
diff --git a/Engine/Plugins/Animation/GeneSplicer/Source/GeneSplicerLib/Private/genesplicer/splicedata/GenePoolMask.cpp b/Engine/Plugins/Animation/GeneSplicer/Source/GeneSplicerLib/Private/genesplicer/splicedata/GenePoolMask.cpp
new file mode 100644
--- /dev/null
+++ b/Engine/Plugins/Animation/GeneSplicer/Source/GeneSplicerLib/Private/genesplicer/splicedata/GenePoolMask.cpp
@@ -0,0 +1,23 @@
+// Copyright Epic Games, Inc. All Rights Reserved.
+
+#include "genesplicer/splicedata/GenePoolMask.h"
+
+#include <type_traits>
+
+namespace gs4 {
+
+GenePoolMask operator~(GenePoolMask a) {
+    return static_cast<GenePoolMask>(~static_cast<std::uint32_t>(a));
+}
+
+GenePoolMask operator|(GenePoolMask a, GenePoolMask b) {
+    return static_cast<GenePoolMask>(static_cast<std::uint32_t>(a) |
+                                     static_cast<std::uint32_t>(b));
+}
+
+GenePoolMask operator&(GenePoolMask a, GenePoolMask b) {
+    return static_cast<GenePoolMask>(static_cast<std::uint32_t>(a) &
+                                     static_cast<std::uint32_t>(b));
+}
+
+}  // namespace gs4
diff --git a/Engine/Plugins/Animation/GeneSplicer/Source/GeneSplicerLib/Private/genesplicer/splicedata/rawgenes/JointGroupOutputIndicesMerger.h b/Engine/Plugins/Animation/GeneSplicer/Source/GeneSplicerLib/Private/genesplicer/splicedata/rawgenes/JointGroupOutputIndicesMerger.h
new file mode 100644
--- /dev/null
+++ b/Engine/Plugins/Animation/GeneSplicer/Source/GeneSplicerLib/Private/genesplicer/splicedata/rawgenes/JointGroupOutputIndicesMerger.h
@@ -0,0 +1,71 @@
+// Copyright Epic Games, Inc. All Rights Reserved.
+
+#pragma once
+
+#include "genesplicer/TypeDefs.h"
+#include "genesplicer/dna/Aliases.h"
+#include "genesplicer/splicedata/rawgenes/RawGenesUtils.h"
+
+#include <cstdint>
+
+namespace gs4 {
+
+struct JointGroupOutputIndicesMerger {
+    public:
+        JointGroupOutputIndicesMerger(ConstArrayView<std::uint16_t> jointIndices_, MemoryResource* memRes_) :
+            memRes{memRes_},
+            jointIndices{jointIndices_},
+            outputIndices{memRes_},
+            lods{memRes_} {
+        }
+
+        void add(const RawJointGroup& jointGroup) {
+            add(ConstArrayView<std::uint16_t>{jointGroup.outputIndices},
+                ConstArrayView<std::uint16_t>{jointGroup.lods});
+        }
+
+        void add(ConstArrayView<std::uint16_t> outputIndices_, ConstArrayView<std::uint16_t> lods_) {
+            outputIndices.push_back(outputIndices_);
+            lods.push_back(lods_);
+        }
+
+        void reserve(std::uint16_t size) {
+            outputIndices.reserve(size);
+            lods.reserve(size);
+        }
+
+        template<typename TOutputIndicesOuputIter, typename ULODReverseOutputIter>
+        TOutputIndicesOuputIter merge(TOutputIndicesOuputIter outputIndicesIter, ULODReverseOutputIter lodsOnePastLast) {
+            Vector<ConstArrayView<std::uint16_t> > outputIndicesForLODVec{outputIndices.size(), {}, memRes};
+            ArrayView<ConstArrayView<std::uint16_t> > outputIndicesForLOD{outputIndicesForLODVec};
+            auto maxOutputIndex = static_cast<std::uint16_t>(
+                (*std::max_element(jointIndices.begin(), jointIndices.end()) + 1u) * 9u);
+
+            auto onePastLastAdded = outputIndicesIter;
+            for (std::size_t lodIdxPlusOne = lods[0].size(); lodIdxPlusOne > 0; --lodIdxPlusOne) {
+                auto lodIdx = static_cast<std::uint16_t>(lodIdxPlusOne - 1u);
+
+                for (std::uint16_t i = 0; i < outputIndices.size(); i++) {
+                    outputIndicesForLOD[i] = getOutputIndicesIntroducedByLOD(outputIndices[i],
+                                                                             lods[i],
+                                                                             lodIdx);
+                }
+                onePastLastAdded = mergeIndices({outputIndicesForLOD},
+                                                maxOutputIndex,
+                                                onePastLastAdded,
+                                                memRes);
+                lodsOnePastLast--;
+                * lodsOnePastLast = static_cast<std::uint16_t>(std::distance(outputIndicesIter, onePastLastAdded));
+            }
+            return onePastLastAdded;
+
+        }
+
+    private:
+        MemoryResource* memRes;
+        ConstArrayView<std::uint16_t> jointIndices;
+        Vector<ConstArrayView<std::uint16_t> > outputIndices;
+        Vector<ConstArrayView<std::uint16_t> > lods;
+};
+
+}  // namespace gs4
diff --git a/Engine/Plugins/Animation/GeneSplicer/Source/GeneSplicerLib/Private/genesplicer/system/Detect.h b/Engine/Plugins/Animation/GeneSplicer/Source/GeneSplicerLib/Private/genesplicer/system/Detect.h
new file mode 100644
--- /dev/null
+++ b/Engine/Plugins/Animation/GeneSplicer/Source/GeneSplicerLib/Private/genesplicer/system/Detect.h
@@ -0,0 +1,12 @@
+// Copyright Epic Games, Inc. All Rights Reserved.
+
+#pragma once
+
+#if defined(GS_AUTODETECT_SSE) && !defined(GS_BUILD_WITH_SSE)
+    // Assume any x86 architecture we try to build on will support SSE
+    #if defined(i386) || defined(__i386) || defined(__i386__) || defined(__IA32__) || defined(_M_IX86) || defined(__X86__) || \
+    defined(_X86_) || defined(__THW_INTEL__) || defined(__I86__) || defined(__amd64__) || defined(__amd64) || \
+    defined(__x86_64__) || defined(__x86_64) || defined(_M_X64) || defined(_M_AMD64)
+        #define GS_BUILD_WITH_SSE 1
+    #endif
+#endif  // GS_AUTODETECT_SSE
diff --git a/Engine/Plugins/Animation/GeneSplicer/Source/GeneSplicerLib/Public/genesplicer/CalculationType.h b/Engine/Plugins/Animation/GeneSplicer/Source/GeneSplicerLib/Public/genesplicer/CalculationType.h
new file mode 100644
--- /dev/null
+++ b/Engine/Plugins/Animation/GeneSplicer/Source/GeneSplicerLib/Public/genesplicer/CalculationType.h
@@ -0,0 +1,13 @@
+// Copyright Epic Games, Inc. All Rights Reserved.
+
+#pragma once
+
+namespace gs4 {
+
+enum class CalculationType {
+    Scalar,  ///< scalar implementation
+    SSE,  ///< vectorized implementation
+    AVX  ///< vectorized implementation
+};
+
+}  // namespace gs4
diff --git a/Engine/Plugins/Animation/GeneSplicer/Source/GeneSplicerLib/Public/genesplicer/version/Version.h b/Engine/Plugins/Animation/GeneSplicer/Source/GeneSplicerLib/Public/genesplicer/version/Version.h
new file mode 100644
--- /dev/null
+++ b/Engine/Plugins/Animation/GeneSplicer/Source/GeneSplicerLib/Public/genesplicer/version/Version.h
@@ -0,0 +1,8 @@
+// Copyright Epic Games, Inc. All Rights Reserved.
+
+#pragma once
+
+#define GS_MAJOR_VERSION 9
+#define GS_MINOR_VERSION 8
+#define GS_PATCH_VERSION 0
+#define GS_VERSION_STRING "9.8.0"
diff --git a/Engine/Plugins/Animation/GeneSplicer/Source/GeneSplicerLib/Public/raf/version/Version.h b/Engine/Plugins/Animation/GeneSplicer/Source/GeneSplicerLib/Public/raf/version/Version.h
new file mode 100644
--- /dev/null
+++ b/Engine/Plugins/Animation/GeneSplicer/Source/GeneSplicerLib/Public/raf/version/Version.h
@@ -0,0 +1,8 @@
+// Copyright Epic Games, Inc. All Rights Reserved.
+
+#pragma once
+
+#define RAF_MAJOR_VERSION 5
+#define RAF_MINOR_VERSION 1
+#define RAF_PATCH_VERSION 7
+#define RAF_VERSION_STRING "5.1.7"
diff --git a/Engine/Plugins/Animation/GeneSplicer/Source/GeneSplicerLibTest/Private/GeneSplicerLibTestSuite.cpp b/Engine/Plugins/Animation/GeneSplicer/Source/GeneSplicerLibTest/Private/GeneSplicerLibTestSuite.cpp
new file mode 100644
--- /dev/null
+++ b/Engine/Plugins/Animation/GeneSplicer/Source/GeneSplicerLibTest/Private/GeneSplicerLibTestSuite.cpp
@@ -0,0 +1,71 @@
+// Copyright Epic Games, Inc. All Rights Reserved.
+
+#if WITH_DEV_AUTOMATION_TESTS
+
+#include "GeneSplicerLibTest.h"
+#include "gtest/gtest.h"
+#include "Misc/AutomationTest.h"
+
+class FGeneSplicerLibTestPrinter
+	: public ::testing::EmptyTestEventListener
+{
+	virtual void OnTestStart(const ::testing::TestInfo& InTestInfo) override
+	{
+		UE_LOG(LogGeneSplicerLibTest, Verbose, TEXT("Test %s.%s Starting"), *FString(InTestInfo.test_suite_name()), *FString(InTestInfo.name()));
+	}
+
+	virtual void OnTestPartResult(const ::testing::TestPartResult& InTestPartResult) override
+	{
+		if (InTestPartResult.failed())
+		{
+			UE_LOG(LogGeneSplicerLibTest, Error, TEXT("FAILED in %s:%d\n%s"), *FString(InTestPartResult.file_name()), InTestPartResult.line_number(), *FString(InTestPartResult.summary()));
+		}
+		else
+		{
+			UE_LOG(LogGeneSplicerLibTest, Verbose, TEXT("Succeeded in %s:%d\n%s"), *FString(InTestPartResult.file_name()), InTestPartResult.line_number(), *FString(InTestPartResult.summary()));
+		}
+	}
+
+	virtual void OnTestEnd(const ::testing::TestInfo& InTestInfo) override
+	{
+		UE_LOG(LogGeneSplicerLibTest, Verbose, TEXT("Test %s.%s Ending"), *FString(InTestInfo.test_suite_name()), *FString(InTestInfo.name()));
+	}
+};
+
+IMPLEMENT_COMPLEX_AUTOMATION_TEST(FGeneSplicerLibTestSuite, "GeneSplicerLib", EAutomationTestFlags::EditorContext | EAutomationTestFlags::EngineFilter)
+
+void FGeneSplicerLibTestSuite::GetTests(TArray<FString>& OutBeautifiedNames, TArray<FString>& OutTestCommands) const
+{
+	::testing::InitGoogleTest();
+
+	const ::testing::UnitTest* Instance = ::testing::UnitTest::GetInstance();
+	for (int32 TestSuiteIndex = 0; TestSuiteIndex < Instance->total_test_suite_count(); ++TestSuiteIndex)
+	{
+		const ::testing::TestSuite* TestSuite = Instance->GetTestSuite(TestSuiteIndex);
+		const FString TestCaseName = FString::Format(TEXT("{0}"), { TestSuite->name() });
+		OutBeautifiedNames.Add(TestCaseName);
+	}
+
+	OutTestCommands = OutBeautifiedNames;	
+}
+
+bool FGeneSplicerLibTestSuite::RunTest(const FString& Parameters)
+{
+	::testing::TestEventListeners& Listeners = ::testing::UnitTest::GetInstance()->listeners();
+
+	FGeneSplicerLibTestPrinter TestPrinter;
+	Listeners.Append(&TestPrinter);
+
+	const FString Filter = FString::Format(TEXT("{0}*"), { Parameters });
+	const TArray<ANSICHAR> TestFilter{ TCHAR_TO_ANSI(*Filter), Filter.Len() + 1 };
+	::testing::GTEST_FLAG(filter) = TestFilter.GetData();
+	return (RUN_ALL_TESTS() == 0);
+
+	TestTrue("GeneSplicerLib Tests", RUN_ALL_TESTS() == 0);
+
+	Listeners.Release(&TestPrinter);
+
+	return true;
+}
+
+#endif  // WITH_DEV_AUTOMATION_TESTS
\ No newline at end of file
diff --git a/Engine/Plugins/Animation/GeneSplicer/Source/GeneSplicerLibTest/Private/gstests/Defs.h b/Engine/Plugins/Animation/GeneSplicer/Source/GeneSplicerLibTest/Private/gstests/Defs.h
new file mode 100644
--- /dev/null
+++ b/Engine/Plugins/Animation/GeneSplicer/Source/GeneSplicerLibTest/Private/gstests/Defs.h
@@ -0,0 +1,75 @@
+// Copyright Epic Games, Inc. All Rights Reserved.
+
+#pragma once
+
+#ifdef _MSC_VER
+    #pragma warning(push)
+    #pragma warning(disable : 4061 4121 4365 4987 4668 4996 6330)
+#endif
+#ifdef __clang__
+    #pragma clang diagnostic push
+    #pragma clang diagnostic ignored "-Wundef"
+#endif
+#include <gtest/gtest.h>
+#ifdef __clang__
+    #pragma clang diagnostic pop
+    // This cannot be applied to the scope of the header inclusion only, as the TEST macros will trigger it
+    #pragma clang diagnostic ignored "-Wglobal-constructors"
+    #pragma clang diagnostic ignored "-Wexit-time-destructors"
+    #pragma clang diagnostic ignored "-Wweak-vtables"
+#endif
+#ifdef _MSC_VER
+    #pragma warning(pop)
+#endif
+
+#ifndef INSTANTIATE_TEST_SUITE_P
+    #define INSTANTIATE_TEST_SUITE_P INSTANTIATE_TEST_CASE_P
+#endif
+
+#ifndef TYPED_TEST_SUITE_P
+    #define TYPED_TEST_SUITE_P TYPED_TEST_CASE_P
+#endif
+
+#ifndef REGISTER_TYPED_TEST_SUITE_P
+    #define REGISTER_TYPED_TEST_SUITE_P REGISTER_TYPED_TEST_CASE_P
+#endif
+
+#ifndef INSTANTIATE_TYPED_TEST_SUITE_P
+    #define INSTANTIATE_TYPED_TEST_SUITE_P INSTANTIATE_TYPED_TEST_CASE_P
+#endif
+
+#ifndef ASSERT_ELEMENTS_EQ
+    #define ASSERT_ELEMENTS_EQ(result, expected, count)     \
+        for (std::size_t i = 0ul; i < count; ++i) {         \
+            ASSERT_EQ(result[i], expected[i]);              \
+        }
+#endif
+
+#ifndef ASSERT_ELEMENTS_AND_SIZE_EQ
+    #define ASSERT_ELEMENTS_AND_SIZE_EQ(result, expected)     \
+        ASSERT_EQ(result.size(), expected.size());              \
+        for (std::size_t i = 0ul; i < result.size(); ++i) {         \
+            ASSERT_EQ(result[i], expected[i]);              \
+        }
+#endif
+
+#ifndef EXPECT_ELEMENTS_EQ
+    #define EXPECT_ELEMENTS_EQ(result, expected, count)     \
+        for (std::size_t i = 0ul; i < count; ++i) {         \
+            EXPECT_EQ(result[i], expected[i]);              \
+        }
+#endif
+
+#ifndef ASSERT_ELEMENTS_NEAR
+    #define ASSERT_ELEMENTS_NEAR(result, expected, count, threshold)    \
+        for (std::size_t i = 0ul; i < count; ++i) {                     \
+            ASSERT_NEAR(result[i], expected[i], threshold);             \
+        }
+#endif
+
+#ifndef EXPECT_ELEMENTS_NEAR
+    #define EXPECT_ELEMENTS_NEAR(result, expected, count, threshold)    \
+        for (std::size_t i = 0ul; i < count; ++i) {                     \
+            EXPECT_NEAR(result[i], expected[i], threshold);             \
+        }
+#endif
