Fix usage of unacceptable words

diff --git a/Engine/Plugins/Animation/GeneSplicer/GeneSplicer.uplugin b/Engine/Plugins/Animation/GeneSplicer/GeneSplicer.uplugin
--- a/Engine/Plugins/Animation/GeneSplicer/GeneSplicer.uplugin
+++ b/Engine/Plugins/Animation/GeneSplicer/GeneSplicer.uplugin
@@ -1,8 +1,8 @@
 {
 	"FileVersion": 3,
 	"Version": 9.8,
-	"VersionName": "9.8.0",
-	"FriendlyName": "GeneSplicer Plugin v9.8.0",
+	"VersionName": "9.8.2",
+	"FriendlyName": "GeneSplicer Plugin v9.8.2",
 	"Description": "3Lateral GeneSplicer plugin for facial animation",
 	"Category": "Animation",
 	"CreatedBy": "3Lateral",
diff --git a/Engine/Plugins/Animation/GeneSplicer/Source/GeneSplicerLib/Private/genesplicer/GeneSplicer.cpp b/Engine/Plugins/Animation/GeneSplicer/Source/GeneSplicerLib/Private/genesplicer/GeneSplicer.cpp
--- a/Engine/Plugins/Animation/GeneSplicer/Source/GeneSplicerLib/Private/genesplicer/GeneSplicer.cpp
+++ b/Engine/Plugins/Animation/GeneSplicer/Source/GeneSplicerLib/Private/genesplicer/GeneSplicer.cpp
@@ -28,28 +28,28 @@ GeneSplicer::~GeneSplicer() = default;
 GeneSplicer::GeneSplicer(GeneSplicer&& rhs) = default;
 GeneSplicer& GeneSplicer::operator=(GeneSplicer&& rhs) = default;
 
-void GeneSplicer::splice(const SpliceData* spliceData, GeneSplicerDNAReader* master) {
-    pImpl->splice(*spliceData, master);
+void GeneSplicer::splice(const SpliceData* spliceData, GeneSplicerDNAReader* output) {
+    pImpl->splice(*spliceData, output);
 }
 
-void GeneSplicer::spliceNeutralMeshes(const SpliceData* spliceData, GeneSplicerDNAReader* master) {
-    pImpl->spliceNeutralMeshes(*spliceData, master);
+void GeneSplicer::spliceNeutralMeshes(const SpliceData* spliceData, GeneSplicerDNAReader* output) {
+    pImpl->spliceNeutralMeshes(*spliceData, output);
 }
 
-void GeneSplicer::spliceBlendShapes(const SpliceData* spliceData, GeneSplicerDNAReader* master) {
-    pImpl->spliceBlendShapes(*spliceData, master);
+void GeneSplicer::spliceBlendShapes(const SpliceData* spliceData, GeneSplicerDNAReader* output) {
+    pImpl->spliceBlendShapes(*spliceData, output);
 }
 
-void GeneSplicer::spliceNeutralJoints(const SpliceData* spliceData, GeneSplicerDNAReader* master) {
-    pImpl->spliceNeutralJoints(*spliceData, master);
+void GeneSplicer::spliceNeutralJoints(const SpliceData* spliceData, GeneSplicerDNAReader* output) {
+    pImpl->spliceNeutralJoints(*spliceData, output);
 }
 
-void GeneSplicer::spliceJointBehavior(const SpliceData* spliceData, GeneSplicerDNAReader* master) {
-    pImpl->spliceJointBehavior(*spliceData, master);
+void GeneSplicer::spliceJointBehavior(const SpliceData* spliceData, GeneSplicerDNAReader* output) {
+    pImpl->spliceJointBehavior(*spliceData, output);
 }
 
-void GeneSplicer::spliceSkinWeights(const SpliceData* spliceData, GeneSplicerDNAReader* master) {
-    pImpl->spliceSkinWeights(*spliceData, master);
+void GeneSplicer::spliceSkinWeights(const SpliceData* spliceData, GeneSplicerDNAReader* output) {
+    pImpl->spliceSkinWeights(*spliceData, output);
 }
 
 }  // namespace gs4
diff --git a/Engine/Plugins/Animation/GeneSplicer/Source/GeneSplicerLib/Private/genesplicer/GeneSplicerImpl.cpp b/Engine/Plugins/Animation/GeneSplicer/Source/GeneSplicerLib/Private/genesplicer/GeneSplicerImpl.cpp
--- a/Engine/Plugins/Animation/GeneSplicer/Source/GeneSplicerLib/Private/genesplicer/GeneSplicerImpl.cpp
+++ b/Engine/Plugins/Animation/GeneSplicer/Source/GeneSplicerLib/Private/genesplicer/GeneSplicerImpl.cpp
@@ -28,32 +28,32 @@ void GeneSplicer::Impl::destroy(GeneSplicer::Impl* instance) {
     alloc.deleteObject(instance);
 }
 
-void GeneSplicer::Impl::splice(const SpliceData& spliceData, GeneSplicerDNAReader* master) {
-    spliceNeutralMeshes(spliceData, master);
-    spliceBlendShapes(spliceData, master);
-    spliceNeutralJoints(spliceData, master);
-    spliceJointBehavior(spliceData, master);
-    spliceSkinWeights(spliceData, master);
+void GeneSplicer::Impl::splice(const SpliceData& spliceData, GeneSplicerDNAReader* output) {
+    spliceNeutralMeshes(spliceData, output);
+    spliceBlendShapes(spliceData, output);
+    spliceNeutralJoints(spliceData, output);
+    spliceJointBehavior(spliceData, output);
+    spliceSkinWeights(spliceData, output);
 }
 
-void GeneSplicer::Impl::spliceNeutralMeshes(const SpliceData& spliceData, GeneSplicerDNAReader* master) {
-    neutralMeshSplicer.get()->splice(PImplExtractor<SpliceData>::get(spliceData), master);
+void GeneSplicer::Impl::spliceNeutralMeshes(const SpliceData& spliceData, GeneSplicerDNAReader* output) {
+    neutralMeshSplicer.get()->splice(PImplExtractor<SpliceData>::get(spliceData), output);
 }
 
-void GeneSplicer::Impl::spliceBlendShapes(const SpliceData& spliceData, GeneSplicerDNAReader* master) {
-    blendShapeSplicer.get()->splice(PImplExtractor<SpliceData>::get(spliceData), master);
+void GeneSplicer::Impl::spliceBlendShapes(const SpliceData& spliceData, GeneSplicerDNAReader* output) {
+    blendShapeSplicer.get()->splice(PImplExtractor<SpliceData>::get(spliceData), output);
 }
 
-void GeneSplicer::Impl::spliceNeutralJoints(const SpliceData& spliceData, GeneSplicerDNAReader* master) {
-    neutralJointSplicer.get()->splice(PImplExtractor<SpliceData>::get(spliceData), master);
+void GeneSplicer::Impl::spliceNeutralJoints(const SpliceData& spliceData, GeneSplicerDNAReader* output) {
+    neutralJointSplicer.get()->splice(PImplExtractor<SpliceData>::get(spliceData), output);
 }
 
-void GeneSplicer::Impl::spliceJointBehavior(const SpliceData& spliceData, GeneSplicerDNAReader* master) {
-    jointBehaviorSplicer.get()->splice(PImplExtractor<SpliceData>::get(spliceData), master);
+void GeneSplicer::Impl::spliceJointBehavior(const SpliceData& spliceData, GeneSplicerDNAReader* output) {
+    jointBehaviorSplicer.get()->splice(PImplExtractor<SpliceData>::get(spliceData), output);
 }
 
-void GeneSplicer::Impl::spliceSkinWeights(const SpliceData& spliceData, GeneSplicerDNAReader* master) {
-    skinWeightSplicer.get()->splice(PImplExtractor<SpliceData>::get(spliceData), master);
+void GeneSplicer::Impl::spliceSkinWeights(const SpliceData& spliceData, GeneSplicerDNAReader* output) {
+    skinWeightSplicer.get()->splice(PImplExtractor<SpliceData>::get(spliceData), output);
 }
 
 }  // namespace gs4
diff --git a/Engine/Plugins/Animation/GeneSplicer/Source/GeneSplicerLib/Private/genesplicer/GeneSplicerImpl.h b/Engine/Plugins/Animation/GeneSplicer/Source/GeneSplicerLib/Private/genesplicer/GeneSplicerImpl.h
--- a/Engine/Plugins/Animation/GeneSplicer/Source/GeneSplicerLib/Private/genesplicer/GeneSplicerImpl.h
+++ b/Engine/Plugins/Animation/GeneSplicer/Source/GeneSplicerLib/Private/genesplicer/GeneSplicerImpl.h
@@ -47,12 +47,12 @@ class GeneSplicer::Impl {
         }
 
         static void destroy(Impl* instance);
-        void splice(const SpliceData& spliceData, GeneSplicerDNAReader* master);
-        void spliceNeutralJoints(const SpliceData& spliceData, GeneSplicerDNAReader* master);
-        void spliceNeutralMeshes(const SpliceData& spliceData, GeneSplicerDNAReader* master);
-        void spliceBlendShapes(const SpliceData& spliceData, GeneSplicerDNAReader* master);
-        void spliceJointBehavior(const SpliceData& spliceData, GeneSplicerDNAReader* master);
-        void spliceSkinWeights(const SpliceData& spliceData, GeneSplicerDNAReader* master);
+        void splice(const SpliceData& spliceData, GeneSplicerDNAReader* output);
+        void spliceNeutralJoints(const SpliceData& spliceData, GeneSplicerDNAReader* output);
+        void spliceNeutralMeshes(const SpliceData& spliceData, GeneSplicerDNAReader* output);
+        void spliceBlendShapes(const SpliceData& spliceData, GeneSplicerDNAReader* output);
+        void spliceJointBehavior(const SpliceData& spliceData, GeneSplicerDNAReader* output);
+        void spliceSkinWeights(const SpliceData& spliceData, GeneSplicerDNAReader* output);
 
     private:
         MemoryResource* memRes;
diff --git a/Engine/Plugins/Animation/GeneSplicer/Source/GeneSplicerLib/Private/genesplicer/Splicer.h b/Engine/Plugins/Animation/GeneSplicer/Source/GeneSplicerLib/Private/genesplicer/Splicer.h
--- a/Engine/Plugins/Animation/GeneSplicer/Source/GeneSplicerLib/Private/genesplicer/Splicer.h
+++ b/Engine/Plugins/Animation/GeneSplicer/Source/GeneSplicerLib/Private/genesplicer/Splicer.h
@@ -17,7 +17,7 @@ class SpliceData;
 class Splicer {
     public:
         explicit Splicer(MemoryResource* memRes);
-        virtual void splice(const SpliceDataInterface* spliceData, GeneSplicerDNAReader* master) = 0;
+        virtual void splice(const SpliceDataInterface* spliceData, GeneSplicerDNAReader* output) = 0;
         virtual ~Splicer();
 
     protected:
diff --git a/Engine/Plugins/Animation/GeneSplicer/Source/GeneSplicerLib/Private/genesplicer/blendshapesplicer/BlendShapeSplicer.cpp b/Engine/Plugins/Animation/GeneSplicer/Source/GeneSplicerLib/Private/genesplicer/blendshapesplicer/BlendShapeSplicer.cpp
--- a/Engine/Plugins/Animation/GeneSplicer/Source/GeneSplicerLib/Private/genesplicer/blendshapesplicer/BlendShapeSplicer.cpp
+++ b/Engine/Plugins/Animation/GeneSplicer/Source/GeneSplicerLib/Private/genesplicer/blendshapesplicer/BlendShapeSplicer.cpp
@@ -88,9 +88,9 @@ class FilteredIterator {
 
 
 template<CalculationType CT>
-void BlendShapeSplicer<CT>::splice(const SpliceDataInterface* spliceData, GeneSplicerDNAReader* master_) {
-    auto master = static_cast<GeneSplicerDNAReaderImpl*>(master_);
-    auto masterMemRes = master->getMemoryResource();
+void BlendShapeSplicer<CT>::splice(const SpliceDataInterface* spliceData, GeneSplicerDNAReader* output_) {
+    auto output = static_cast<GeneSplicerDNAReaderImpl*>(output_);
+    auto outputMemRes = output->getMemoryResource();
 
     const auto& baseArch = spliceData->getBaseArchetype();
     auto baseArchBlendShapes = baseArch.getBlendShapeTargets();
@@ -99,11 +99,11 @@ void BlendShapeSplicer<CT>::splice(const SpliceDataInterface* spliceData, GeneSp
             return pool->getGenePool()->getBlendShapeTargetCount(meshIndex) ==
                    baseArchetype.getBlendShapeTargets()[meshIndex].size();
         };
-    auto poolsToSplicePerMesh = filterPoolSpliceParamsPerMesh(spliceData, blendShapePredicate, masterMemRes);
+    auto poolsToSplicePerMesh = filterPoolSpliceParamsPerMesh(spliceData, blendShapePredicate, outputMemRes);
 
     for (std::uint16_t meshIdx = 0u; meshIdx <  poolsToSplicePerMesh.rowCount(); meshIdx++) {
         Vector<RawBlendShapeTarget> resultingBlendShapeTargets =
-            constructBlendShapeTargetsWithPadding(baseArchBlendShapes[meshIdx], masterMemRes);
+            constructBlendShapeTargetsWithPadding(baseArchBlendShapes[meshIdx], outputMemRes);
 
         auto poolParams = poolsToSplicePerMesh[meshIdx];
 
@@ -179,7 +179,7 @@ void BlendShapeSplicer<CT>::splice(const SpliceDataInterface* spliceData, GeneSp
             }
             resultDeltas.resize(resultIndices.size());
         }
-        master->setBlendShapeTargets(meshIdx, std::move(resultingBlendShapeTargets));
+        output->setBlendShapeTargets(meshIdx, std::move(resultingBlendShapeTargets));
     }
 }
 
diff --git a/Engine/Plugins/Animation/GeneSplicer/Source/GeneSplicerLib/Private/genesplicer/blendshapesplicer/BlendShapeSplicer.h b/Engine/Plugins/Animation/GeneSplicer/Source/GeneSplicerLib/Private/genesplicer/blendshapesplicer/BlendShapeSplicer.h
--- a/Engine/Plugins/Animation/GeneSplicer/Source/GeneSplicerLib/Private/genesplicer/blendshapesplicer/BlendShapeSplicer.h
+++ b/Engine/Plugins/Animation/GeneSplicer/Source/GeneSplicerLib/Private/genesplicer/blendshapesplicer/BlendShapeSplicer.h
@@ -17,7 +17,7 @@ class BlendShapeSplicer : public Splicer {
     public:
         using Splicer::Splicer;
 
-        void splice(const SpliceDataInterface* spliceData, GeneSplicerDNAReader* master) override;
+        void splice(const SpliceDataInterface* spliceData, GeneSplicerDNAReader* output) override;
 
 };
 
diff --git a/Engine/Plugins/Animation/GeneSplicer/Source/GeneSplicerLib/Private/genesplicer/geometrysplicer/NeutralMeshSplicer.cpp b/Engine/Plugins/Animation/GeneSplicer/Source/GeneSplicerLib/Private/genesplicer/geometrysplicer/NeutralMeshSplicer.cpp
--- a/Engine/Plugins/Animation/GeneSplicer/Source/GeneSplicerLib/Private/genesplicer/geometrysplicer/NeutralMeshSplicer.cpp
+++ b/Engine/Plugins/Animation/GeneSplicer/Source/GeneSplicerLib/Private/genesplicer/geometrysplicer/NeutralMeshSplicer.cpp
@@ -17,9 +17,9 @@
 namespace gs4 {
 
 template<CalculationType CT>
-void NeutralMeshSplicer<CT>::splice(const SpliceDataInterface* spliceData, GeneSplicerDNAReader* master_) {
-    auto master = static_cast<GeneSplicerDNAReaderImpl*>(master_);
-    auto masterMemRes = master->getMemoryResource();
+void NeutralMeshSplicer<CT>::splice(const SpliceDataInterface* spliceData, GeneSplicerDNAReader* output_) {
+    auto output = static_cast<GeneSplicerDNAReaderImpl*>(output_);
+    auto outputMemRes = output->getMemoryResource();
 
     const auto& baseArch = spliceData->getBaseArchetype();
     const auto& baseArchMeshes = baseArch.getNeutralMeshes();
@@ -27,12 +27,12 @@ void NeutralMeshSplicer<CT>::splice(const SpliceDataInterface* spliceData, GeneS
     auto neutralMeshPredicate = [](const RawGenes& baseArchetype, const PoolSpliceParamsImpl* pool, std::uint16_t meshIndex) {
             return baseArchetype.getNeutralMeshes()[meshIndex].size() == pool->getGenePool()->getVertexPositionCount(meshIndex);
         };
-    auto poolsToSplicePerMesh = filterPoolSpliceParamsPerMesh(spliceData, neutralMeshPredicate, masterMemRes);
+    auto poolsToSplicePerMesh = filterPoolSpliceParamsPerMesh(spliceData, neutralMeshPredicate, outputMemRes);
 
     for (std::uint16_t meshIdx = 0u; meshIdx <  poolsToSplicePerMesh.rowCount(); meshIdx++) {
 
         const auto& baseArchMesh = baseArchMeshes[meshIdx];
-        RawVector3Vector resultingVertices = constructWithPadding(baseArchMesh, masterMemRes);
+        RawVector3Vector resultingVertices = constructWithPadding(baseArchMesh, outputMemRes);
 
         for (auto poolParams : poolsToSplicePerMesh[meshIdx]) {
             auto genePool = poolParams->getGenePool();
@@ -46,8 +46,8 @@ void NeutralMeshSplicer<CT>::splice(const SpliceDataInterface* spliceData, GeneS
                                      resultingVertices,
                                      scale);
         }
-        resultingVertices.resize(master->getVertexPositionCount(meshIdx));
-        master->setVertexPositions(meshIdx, std::move(resultingVertices));
+        resultingVertices.resize(output->getVertexPositionCount(meshIdx));
+        output->setVertexPositions(meshIdx, std::move(resultingVertices));
     }
 }
 
diff --git a/Engine/Plugins/Animation/GeneSplicer/Source/GeneSplicerLib/Private/genesplicer/geometrysplicer/NeutralMeshSplicer.h b/Engine/Plugins/Animation/GeneSplicer/Source/GeneSplicerLib/Private/genesplicer/geometrysplicer/NeutralMeshSplicer.h
--- a/Engine/Plugins/Animation/GeneSplicer/Source/GeneSplicerLib/Private/genesplicer/geometrysplicer/NeutralMeshSplicer.h
+++ b/Engine/Plugins/Animation/GeneSplicer/Source/GeneSplicerLib/Private/genesplicer/geometrysplicer/NeutralMeshSplicer.h
@@ -18,7 +18,7 @@ class NeutralMeshSplicer : public Splicer {
     public:
         using Splicer::Splicer;
 
-        void splice(const SpliceDataInterface* spliceData, GeneSplicerDNAReader* master) override;
+        void splice(const SpliceDataInterface* spliceData, GeneSplicerDNAReader* output) override;
 };
 
 }  // namespace gs4
diff --git a/Engine/Plugins/Animation/GeneSplicer/Source/GeneSplicerLib/Private/genesplicer/jointbehaviorsplicer/JointBehaviorSplicer.cpp b/Engine/Plugins/Animation/GeneSplicer/Source/GeneSplicerLib/Private/genesplicer/jointbehaviorsplicer/JointBehaviorSplicer.cpp
--- a/Engine/Plugins/Animation/GeneSplicer/Source/GeneSplicerLib/Private/genesplicer/jointbehaviorsplicer/JointBehaviorSplicer.cpp
+++ b/Engine/Plugins/Animation/GeneSplicer/Source/GeneSplicerLib/Private/genesplicer/jointbehaviorsplicer/JointBehaviorSplicer.cpp
@@ -133,16 +133,16 @@ struct OutputIndexSplicer<CalculationType::SSE> {
 }  // namespace
 
 template<CalculationType CT>
-void JointBehaviorSplicer<CT>::splice(const SpliceDataInterface* spliceData, GeneSplicerDNAReader* master_) {
-    auto master = static_cast<GeneSplicerDNAReaderImpl*>(master_);
-    auto masterMemRes = master->getMemoryResource();
+void JointBehaviorSplicer<CT>::splice(const SpliceDataInterface* spliceData, GeneSplicerDNAReader* output_) {
+    auto output = static_cast<GeneSplicerDNAReaderImpl*>(output_);
+    auto outputMemRes = output->getMemoryResource();
 
     const auto& baseArchJointGroups = spliceData->getBaseArchetype().getJointGroups();
-    Vector<RawJointGroup> resultingJointGroups{baseArchJointGroups.begin(), baseArchJointGroups.end(), masterMemRes};
+    Vector<RawJointGroup> resultingJointGroups{baseArchJointGroups.begin(), baseArchJointGroups.end(), outputMemRes};
     auto pools = spliceData->getAllPoolParams();
 
     std::uint16_t jointCount = spliceData->getBaseArchetype().getJointCount();
-    Vector<std::uint16_t> jointToGroupMapping{jointCount, {}, masterMemRes};
+    Vector<std::uint16_t> jointToGroupMapping{jointCount, {}, outputMemRes};
 
     for (std::uint16_t jointGroupIdx = 0u; jointGroupIdx < resultingJointGroups.size(); jointGroupIdx++) {
         for (auto jointIndex : resultingJointGroups[jointGroupIdx].jointIndices) {
@@ -206,7 +206,7 @@ void JointBehaviorSplicer<CT>::splice(const SpliceDataInterface* spliceData, Gen
             }
         }
     }
-    master->setJointGroups(std::move(resultingJointGroups));
+    output->setJointGroups(std::move(resultingJointGroups));
 }
 
 template class JointBehaviorSplicer<CalculationType::Scalar>;
diff --git a/Engine/Plugins/Animation/GeneSplicer/Source/GeneSplicerLib/Private/genesplicer/jointbehaviorsplicer/JointBehaviorSplicer.h b/Engine/Plugins/Animation/GeneSplicer/Source/GeneSplicerLib/Private/genesplicer/jointbehaviorsplicer/JointBehaviorSplicer.h
--- a/Engine/Plugins/Animation/GeneSplicer/Source/GeneSplicerLib/Private/genesplicer/jointbehaviorsplicer/JointBehaviorSplicer.h
+++ b/Engine/Plugins/Animation/GeneSplicer/Source/GeneSplicerLib/Private/genesplicer/jointbehaviorsplicer/JointBehaviorSplicer.h
@@ -14,7 +14,7 @@ class JointBehaviorSplicer : public Splicer {
     public:
         using Splicer::Splicer;
 
-        void splice(const SpliceDataInterface* spliceData, GeneSplicerDNAReader* master) override;
+        void splice(const SpliceDataInterface* spliceData, GeneSplicerDNAReader* output) override;
 
 };
 
diff --git a/Engine/Plugins/Animation/GeneSplicer/Source/GeneSplicerLib/Private/genesplicer/neutraljointsplicer/JointAttributeSpecialization.h b/Engine/Plugins/Animation/GeneSplicer/Source/GeneSplicerLib/Private/genesplicer/neutraljointsplicer/JointAttributeSpecialization.h
--- a/Engine/Plugins/Animation/GeneSplicer/Source/GeneSplicerLib/Private/genesplicer/neutraljointsplicer/JointAttributeSpecialization.h
+++ b/Engine/Plugins/Animation/GeneSplicer/Source/GeneSplicerLib/Private/genesplicer/neutraljointsplicer/JointAttributeSpecialization.h
@@ -38,17 +38,17 @@ FORCE_INLINE Vec3 getJointAttributeValue<JointAttribute::Rotation>(const Reader*
 }
 
 template<JointAttribute T>
-FORCE_INLINE void setJointAttributeValues(GeneSplicerDNAReaderImpl* master, RawVector3Vector& result);
+FORCE_INLINE void setJointAttributeValues(GeneSplicerDNAReaderImpl* output, RawVector3Vector& result);
 
 template<>
-FORCE_INLINE void setJointAttributeValues<JointAttribute::Translation>(GeneSplicerDNAReaderImpl* master,
+FORCE_INLINE void setJointAttributeValues<JointAttribute::Translation>(GeneSplicerDNAReaderImpl* output,
                                                                        RawVector3Vector& result) {
-    master->setNeutralJointTranslations(std::move(result));
+    output->setNeutralJointTranslations(std::move(result));
 }
 
 template<>
-FORCE_INLINE void setJointAttributeValues<JointAttribute::Rotation>(GeneSplicerDNAReaderImpl* master, RawVector3Vector& result) {
-    master->setNeutralJointRotations(std::move(result));
+FORCE_INLINE void setJointAttributeValues<JointAttribute::Rotation>(GeneSplicerDNAReaderImpl* output, RawVector3Vector& result) {
+    output->setNeutralJointRotations(std::move(result));
 }
 
 }  // namespace gs4
diff --git a/Engine/Plugins/Animation/GeneSplicer/Source/GeneSplicerLib/Private/genesplicer/neutraljointsplicer/NeutralJointSplicer.cpp b/Engine/Plugins/Animation/GeneSplicer/Source/GeneSplicerLib/Private/genesplicer/neutraljointsplicer/NeutralJointSplicer.cpp
--- a/Engine/Plugins/Animation/GeneSplicer/Source/GeneSplicerLib/Private/genesplicer/neutraljointsplicer/NeutralJointSplicer.cpp
+++ b/Engine/Plugins/Animation/GeneSplicer/Source/GeneSplicerLib/Private/genesplicer/neutraljointsplicer/NeutralJointSplicer.cpp
@@ -27,16 +27,16 @@
 namespace gs4 {
 
 template<CalculationType CT>
-void NeutralJointSplicer<CT>::splice(const SpliceDataInterface* spliceData, GeneSplicerDNAReader* master_) {
-    auto master = static_cast<GeneSplicerDNAReaderImpl*>(master_);
+void NeutralJointSplicer<CT>::splice(const SpliceDataInterface* spliceData, GeneSplicerDNAReader* output_) {
+    auto output = static_cast<GeneSplicerDNAReaderImpl*>(output_);
     auto splice =
-        [spliceData, master](JointAttribute jointAttribute,
+        [spliceData, output](JointAttribute jointAttribute,
                              std::function<float(const PoolSpliceParamsImpl* poolParams)> scaleGetter) {
-            auto masterMemRes = master->getMemoryResource();
+            auto outputMemRes = output->getMemoryResource();
             const auto& baseArch = spliceData->getBaseArchetype();
             auto poolParams = spliceData->getAllPoolParams();
             const auto& baseArchJoints = baseArch.getNeutralJoints(jointAttribute);
-            RawVector3Vector resultingJoints = constructWithPadding(baseArchJoints, masterMemRes);
+            RawVector3Vector resultingJoints = constructWithPadding(baseArchJoints, outputMemRes);
 
             for (auto pool : poolParams) {
                 auto genePool = pool->getGenePool();
@@ -55,7 +55,7 @@ void NeutralJointSplicer<CT>::splice(const SpliceDataInterface* spliceData, Gene
             resultingJoints.resize(baseArch.getJointCount());
             return resultingJoints;
         };
-    RawNeutralJoints neutralJoints{master->getMemoryResource()};
+    RawNeutralJoints neutralJoints{output->getMemoryResource()};
     neutralJoints.translations = splice(JointAttribute::Translation, [](const PoolSpliceParamsImpl* poolParams) {
             return poolParams->getScale();
         });
@@ -63,7 +63,7 @@ void NeutralJointSplicer<CT>::splice(const SpliceDataInterface* spliceData, Gene
     neutralJoints.rotations = splice(JointAttribute::Rotation, [](const PoolSpliceParamsImpl*  /*unused*/) {
             return 1.0f;
         });
-    auto getJointParentIndex = std::bind(&dna::DefinitionReader::getJointParentIndex, master_, std::placeholders::_1);
+    auto getJointParentIndex = std::bind(&dna::DefinitionReader::getJointParentIndex, output_, std::placeholders::_1);
     // Since GenePools and RawGenes of archetype hold joints in world space we need to convert them to local space
     toLocalSpace(getJointParentIndex, neutralJoints);
 
@@ -76,8 +76,8 @@ void NeutralJointSplicer<CT>::splice(const SpliceDataInterface* spliceData, Gene
     std::for_each(rotations.ys.begin(), rotations.ys.end(), toDegrees);
     std::for_each(rotations.zs.begin(), rotations.zs.end(), toDegrees);
 
-    master->setNeutralJointTranslations(std::move(neutralJoints.translations));
-    master->setNeutralJointRotations(std::move(neutralJoints.rotations));
+    output->setNeutralJointTranslations(std::move(neutralJoints.translations));
+    output->setNeutralJointRotations(std::move(neutralJoints.rotations));
 }
 
 template class NeutralJointSplicer<CalculationType::Scalar>;
diff --git a/Engine/Plugins/Animation/GeneSplicer/Source/GeneSplicerLib/Private/genesplicer/neutraljointsplicer/NeutralJointSplicer.h b/Engine/Plugins/Animation/GeneSplicer/Source/GeneSplicerLib/Private/genesplicer/neutraljointsplicer/NeutralJointSplicer.h
--- a/Engine/Plugins/Animation/GeneSplicer/Source/GeneSplicerLib/Private/genesplicer/neutraljointsplicer/NeutralJointSplicer.h
+++ b/Engine/Plugins/Animation/GeneSplicer/Source/GeneSplicerLib/Private/genesplicer/neutraljointsplicer/NeutralJointSplicer.h
@@ -15,7 +15,7 @@ class NeutralJointSplicer : public Splicer {
     public:
         using Splicer::Splicer;
 
-        void splice(const SpliceDataInterface* spliceData, GeneSplicerDNAReader* master) override;
+        void splice(const SpliceDataInterface* spliceData, GeneSplicerDNAReader* output) override;
 
 };
 
diff --git a/Engine/Plugins/Animation/GeneSplicer/Source/GeneSplicerLib/Private/genesplicer/skinweightsplicer/SkinWeightMeshSplicer.cpp b/Engine/Plugins/Animation/GeneSplicer/Source/GeneSplicerLib/Private/genesplicer/skinweightsplicer/SkinWeightMeshSplicer.cpp
--- a/Engine/Plugins/Animation/GeneSplicer/Source/GeneSplicerLib/Private/genesplicer/skinweightsplicer/SkinWeightMeshSplicer.cpp
+++ b/Engine/Plugins/Animation/GeneSplicer/Source/GeneSplicerLib/Private/genesplicer/skinweightsplicer/SkinWeightMeshSplicer.cpp
@@ -52,36 +52,36 @@ SkinWeightMeshSplicer<CT>::GenePoolDetails::GenePoolDetails(const PoolSplicePara
 
 template<CalculationType CT>
 Vector<RawVertexSkinWeights> SkinWeightMeshSplicer<CT>::spliceMesh(const Vector<RawVertexSkinWeights>& baseArchSkinWeights,
-                                                                   MemoryResource* masterMemRes) const {
-    Vector<RawVertexSkinWeights> masterSkinWeights{baseArchSkinWeights.begin(),
+                                                                   MemoryResource* outputMemRes) const {
+    Vector<RawVertexSkinWeights> outputSkinWeights{baseArchSkinWeights.begin(),
                         baseArchSkinWeights.end(),
-                        masterMemRes};
+                        outputMemRes};
     const std::size_t blockCount = getBlockCount(skinWeightCount);
     constexpr std::uint16_t blockCapacity = 16u;
     if (blockCount > 0) {
         std::uint32_t blockIndex = 0u;
         for (; blockIndex < blockCount - 1u; ++blockIndex) {
-            spliceBlock(blockIndex, {masterSkinWeights.data() + blockIndex* blockCapacity, blockCapacity});
+            spliceBlock(blockIndex, {outputSkinWeights.data() + blockIndex* blockCapacity, blockCapacity});
         }
         const auto blockSize = skinWeightCount - blockIndex * blockCapacity;
-        spliceBlock(blockIndex, {masterSkinWeights.data() + blockIndex* blockCapacity, blockSize});
+        spliceBlock(blockIndex, {outputSkinWeights.data() + blockIndex* blockCapacity, blockSize});
     }
-    return masterSkinWeights;
+    return outputSkinWeights;
 }
 
 template<CalculationType CT>
-void SkinWeightMeshSplicer<CT>::spliceBlock(std::uint32_t blockIndex, ArrayView<RawVertexSkinWeights> masterBlock) const {
+void SkinWeightMeshSplicer<CT>::spliceBlock(std::uint32_t blockIndex, ArrayView<RawVertexSkinWeights> outputBlock) const {
     auto maximumJointCount = static_cast<std::uint16_t>(genePools[0].meshDNABlocks[blockIndex].columnCount());
     AlignedVBlock16Vector weights{maximumJointCount, {}, getMemoryResource()};
     const auto weightSum = spliceAndNormalize(blockIndex, weights);
 
     constexpr std::uint16_t blockSize = 16u;
     std::uint32_t vtxIndex = blockIndex * blockSize;
-    for (std::uint32_t vtxOffset = 0; vtxOffset < masterBlock.size(); vtxOffset++) {
+    for (std::uint32_t vtxOffset = 0; vtxOffset < outputBlock.size(); vtxOffset++) {
         if (weightSum[vtxOffset] == 0.0f) {
             continue;
         }
-        auto& vtxRes = masterBlock[vtxOffset];
+        auto& vtxRes = outputBlock[vtxOffset];
         auto vtxJointIndices = getJointIndices(vtxIndex + vtxOffset);
         vtxRes.jointIndices.assign(vtxJointIndices.begin(), vtxJointIndices.end());
         vtxRes.weights.resize(vtxRes.jointIndices.size());
diff --git a/Engine/Plugins/Animation/GeneSplicer/Source/GeneSplicerLib/Private/genesplicer/skinweightsplicer/SkinWeightMeshSplicer.h b/Engine/Plugins/Animation/GeneSplicer/Source/GeneSplicerLib/Private/genesplicer/skinweightsplicer/SkinWeightMeshSplicer.h
--- a/Engine/Plugins/Animation/GeneSplicer/Source/GeneSplicerLib/Private/genesplicer/skinweightsplicer/SkinWeightMeshSplicer.h
+++ b/Engine/Plugins/Animation/GeneSplicer/Source/GeneSplicerLib/Private/genesplicer/skinweightsplicer/SkinWeightMeshSplicer.h
@@ -32,7 +32,7 @@ class SkinWeightMeshSplicer {
     public:
         SkinWeightMeshSplicer(ConstArrayView<const PoolSpliceParamsImpl*> pools, std::uint16_t meshIndex, MemoryResource* memRes);
         Vector<RawVertexSkinWeights> spliceMesh(const Vector<RawVertexSkinWeights>& baseArchetypeMeshSkinWeights,
-                                                MemoryResource* masterMemRes) const;
+                                                MemoryResource* outputMemRes) const;
 
     private:
         ConstArrayView<std::uint16_t> getJointIndices(std::uint32_t vertexIndex) const;
@@ -41,7 +41,7 @@ class SkinWeightMeshSplicer {
         void prune(RawVertexSkinWeights& vtxRes) const;
         VBlock<16> spliceAndNormalize(std::uint32_t blockIndex, AlignedVBlock16Vector& blockResult) const;
 
-        void spliceBlock(std::uint32_t blockIndex, ArrayView<RawVertexSkinWeights> masterBlock) const;
+        void spliceBlock(std::uint32_t blockIndex, ArrayView<RawVertexSkinWeights> outputBlock) const;
 
     private:
         Vector<GenePoolDetails> genePools;
diff --git a/Engine/Plugins/Animation/GeneSplicer/Source/GeneSplicerLib/Private/genesplicer/skinweightsplicer/SkinWeightSplicer.cpp b/Engine/Plugins/Animation/GeneSplicer/Source/GeneSplicerLib/Private/genesplicer/skinweightsplicer/SkinWeightSplicer.cpp
--- a/Engine/Plugins/Animation/GeneSplicer/Source/GeneSplicerLib/Private/genesplicer/skinweightsplicer/SkinWeightSplicer.cpp
+++ b/Engine/Plugins/Animation/GeneSplicer/Source/GeneSplicerLib/Private/genesplicer/skinweightsplicer/SkinWeightSplicer.cpp
@@ -33,22 +33,22 @@
 namespace gs4 {
 
 template<CalculationType CT>
-void SkinWeightSplicer<CT>::splice(const SpliceDataInterface* spliceData, GeneSplicerDNAReader* master_) {
-    auto master = static_cast<GeneSplicerDNAReaderImpl*>(master_);
-    auto masterMemRes = master->getMemoryResource();
+void SkinWeightSplicer<CT>::splice(const SpliceDataInterface* spliceData, GeneSplicerDNAReader* output_) {
+    auto output = static_cast<GeneSplicerDNAReaderImpl*>(output_);
+    auto outputMemRes = output->getMemoryResource();
 
     const auto& baseArch = spliceData->getBaseArchetype();
 
     auto skinWeightPredicate = [](const RawGenes& baseArchetype, const PoolSpliceParamsImpl* pool, std::uint16_t meshIndex) {
             return baseArchetype.getSkinWeightsCount(meshIndex) == pool->getGenePool()->getSkinWeightsCount(meshIndex);
         };
-    auto poolsToSplice = filterPoolSpliceParamsPerMesh(spliceData, skinWeightPredicate, masterMemRes);
+    auto poolsToSplice = filterPoolSpliceParamsPerMesh(spliceData, skinWeightPredicate, outputMemRes);
 
     for (std::uint16_t meshIdx = 0u; meshIdx <  baseArch.getMeshCount(); meshIdx++) {
         auto baseArchSkinWeights = baseArch.getSkinWeights();
-        SkinWeightMeshSplicer<CT> meshSkinWeightPoolViewer{poolsToSplice[meshIdx], meshIdx, masterMemRes};
-        auto meshResult = meshSkinWeightPoolViewer.spliceMesh(baseArchSkinWeights[meshIdx], masterMemRes);
-        master->setSkinWeights(meshIdx, std::move(meshResult));
+        SkinWeightMeshSplicer<CT> meshSkinWeightPoolViewer{poolsToSplice[meshIdx], meshIdx, outputMemRes};
+        auto meshResult = meshSkinWeightPoolViewer.spliceMesh(baseArchSkinWeights[meshIdx], outputMemRes);
+        output->setSkinWeights(meshIdx, std::move(meshResult));
     }
 }
 
diff --git a/Engine/Plugins/Animation/GeneSplicer/Source/GeneSplicerLib/Private/genesplicer/skinweightsplicer/SkinWeightSplicer.h b/Engine/Plugins/Animation/GeneSplicer/Source/GeneSplicerLib/Private/genesplicer/skinweightsplicer/SkinWeightSplicer.h
--- a/Engine/Plugins/Animation/GeneSplicer/Source/GeneSplicerLib/Private/genesplicer/skinweightsplicer/SkinWeightSplicer.h
+++ b/Engine/Plugins/Animation/GeneSplicer/Source/GeneSplicerLib/Private/genesplicer/skinweightsplicer/SkinWeightSplicer.h
@@ -14,7 +14,7 @@ class SkinWeightSplicer : public Splicer {
     public:
         using Splicer::Splicer;
 
-        void splice(const SpliceDataInterface* spliceData, GeneSplicerDNAReader* master) override;
+        void splice(const SpliceDataInterface* spliceData, GeneSplicerDNAReader* output) override;
 
 };
 
diff --git a/Engine/Plugins/Animation/GeneSplicer/Source/GeneSplicerLib/Public/genesplicer/GeneSplicer.h b/Engine/Plugins/Animation/GeneSplicer/Source/GeneSplicerLib/Public/genesplicer/GeneSplicer.h
--- a/Engine/Plugins/Animation/GeneSplicer/Source/GeneSplicerLib/Public/genesplicer/GeneSplicer.h
+++ b/Engine/Plugins/Animation/GeneSplicer/Source/GeneSplicerLib/Public/genesplicer/GeneSplicer.h
@@ -33,50 +33,50 @@ class GeneSplicer final {
             @brief Run all the individual splicers.
             @param spliceData
                 Contains all the necessary input data that is used during splicing.
-            @param master
+            @param output
                 Output parameter - the DNA Reader that will contain the spliced DNA data.
         */
-        GSAPI void splice(const SpliceData* spliceData, GeneSplicerDNAReader* master);
+        GSAPI void splice(const SpliceData* spliceData, GeneSplicerDNAReader* output);
         /**
             @brief Run only the neutral mesh splicer.
             @param spliceData
                 Contains all the necessary input data that is used during splicing.
-            @param master
+            @param output
                 Output parameter - the DNA Reader that will contain the spliced DNA data.
         */
-        GSAPI void spliceNeutralMeshes(const SpliceData* spliceData, GeneSplicerDNAReader* master);
+        GSAPI void spliceNeutralMeshes(const SpliceData* spliceData, GeneSplicerDNAReader* output);
         /**
             @brief Run only the blend shape splicer.
             @param spliceData
                 Contains all the necessary input data that is used during splicing.
-            @param master
+            @param output
                 Output parameter - the DNA Reader that will contain the spliced DNA data.
         */
-        GSAPI void spliceBlendShapes(const SpliceData* spliceData, GeneSplicerDNAReader* master);
+        GSAPI void spliceBlendShapes(const SpliceData* spliceData, GeneSplicerDNAReader* output);
         /**
             @brief Run only the neutral joint splicer.
             @param spliceData
                 Contains all the necessary input data that is used during splicing.
-            @param master
+            @param output
                 Output parameter - the DNA Reader that will contain the spliced DNA data.
         */
-        GSAPI void spliceNeutralJoints(const SpliceData* spliceData, GeneSplicerDNAReader* master);
+        GSAPI void spliceNeutralJoints(const SpliceData* spliceData, GeneSplicerDNAReader* output);
         /**
             @brief Run only the joint behavior splicer.
             @param spliceData
                 Contains all the necessary input data that is used during splicing.
-            @param master
+            @param output
                 Output parameter - the DNA Reader that will contain the spliced DNA data.
         */
-        GSAPI void spliceJointBehavior(const SpliceData* spliceData, GeneSplicerDNAReader* master);
+        GSAPI void spliceJointBehavior(const SpliceData* spliceData, GeneSplicerDNAReader* output);
         /**
             @brief Run only the skin weight splicer.
             @param spliceData
                 Contains all the necessary input data that is used during splicing.
-            @param master
+            @param output
                 Output parameter - the DNA Reader that will contain the spliced DNA data.
         */
-        GSAPI void spliceSkinWeights(const SpliceData* spliceData, GeneSplicerDNAReader* master);
+        GSAPI void spliceSkinWeights(const SpliceData* spliceData, GeneSplicerDNAReader* output);
 
     private:
         class Impl;
diff --git a/Engine/Plugins/Animation/GeneSplicer/Source/GeneSplicerLib/Public/genesplicer/version/Version.h b/Engine/Plugins/Animation/GeneSplicer/Source/GeneSplicerLib/Public/genesplicer/version/Version.h
--- a/Engine/Plugins/Animation/GeneSplicer/Source/GeneSplicerLib/Public/genesplicer/version/Version.h
+++ b/Engine/Plugins/Animation/GeneSplicer/Source/GeneSplicerLib/Public/genesplicer/version/Version.h
@@ -4,5 +4,5 @@
 
 #define GS_MAJOR_VERSION 9
 #define GS_MINOR_VERSION 8
-#define GS_PATCH_VERSION 0
-#define GS_VERSION_STRING "9.8.0"
+#define GS_PATCH_VERSION 2
+#define GS_VERSION_STRING "9.8.2"
diff --git a/Engine/Plugins/Animation/GeneSplicer/Source/GeneSplicerLib/Public/raf/version/Version.h b/Engine/Plugins/Animation/GeneSplicer/Source/GeneSplicerLib/Public/raf/version/Version.h
--- a/Engine/Plugins/Animation/GeneSplicer/Source/GeneSplicerLib/Public/raf/version/Version.h
+++ b/Engine/Plugins/Animation/GeneSplicer/Source/GeneSplicerLib/Public/raf/version/Version.h
@@ -4,5 +4,5 @@
 
 #define RAF_MAJOR_VERSION 5
 #define RAF_MINOR_VERSION 1
-#define RAF_PATCH_VERSION 7
-#define RAF_VERSION_STRING "5.1.7"
+#define RAF_PATCH_VERSION 8
+#define RAF_VERSION_STRING "5.1.8"
diff --git a/Engine/Plugins/Animation/GeneSplicer/Source/GeneSplicerLibTest/Private/gstests/Assertions.h b/Engine/Plugins/Animation/GeneSplicer/Source/GeneSplicerLibTest/Private/gstests/Assertions.h
--- a/Engine/Plugins/Animation/GeneSplicer/Source/GeneSplicerLibTest/Private/gstests/Assertions.h
+++ b/Engine/Plugins/Animation/GeneSplicer/Source/GeneSplicerLibTest/Private/gstests/Assertions.h
@@ -26,12 +26,12 @@ inline void assertVec3Near(Vec3 result, Vec3 expected, float threshold) {
     ASSERT_NEAR(result.z, expected.z, threshold);
 }
 
-inline void assertNeutralMeshes(Reader* master, Reader* expected, float threshold) {
-    ASSERT_EQ(master->getMeshCount(), expected->getMeshCount());
+inline void assertNeutralMeshes(Reader* output, Reader* expected, float threshold) {
+    ASSERT_EQ(output->getMeshCount(), expected->getMeshCount());
     for (std::uint16_t meshIdx = 0u; meshIdx < expected->getMeshCount(); ++meshIdx) {
-        ASSERT_EQ(master->getVertexPositionCount(meshIdx), expected->getVertexPositionCount(meshIdx));
+        ASSERT_EQ(output->getVertexPositionCount(meshIdx), expected->getVertexPositionCount(meshIdx));
         for (std::uint32_t vertexIdx = 0u; vertexIdx < expected->getVertexPositionCount(meshIdx); ++vertexIdx) {
-            Vec3 resultPosition = master->getVertexPosition(meshIdx, vertexIdx);
+            Vec3 resultPosition = output->getVertexPosition(meshIdx, vertexIdx);
             Vec3 expectedPosition = expected->getVertexPosition(meshIdx, vertexIdx);
             assertVec3Near(resultPosition, expectedPosition, threshold);
         }
@@ -74,20 +74,20 @@ inline void assertNeutralMeshPoolData(ConstArrayView<XYZTiledMatrix<16u> > actua
     }
 }
 
-inline void assertBlendShapeTargets(Reader* master, Reader* expected, float threshold) {
-    ASSERT_EQ(master->getMeshCount(), expected->getMeshCount());
+inline void assertBlendShapeTargets(Reader* output, Reader* expected, float threshold) {
+    ASSERT_EQ(output->getMeshCount(), expected->getMeshCount());
     for (std::uint16_t meshIdx = 0u; meshIdx < expected->getMeshCount(); ++meshIdx) {
-        ASSERT_EQ(master->getBlendShapeTargetCount(meshIdx), expected->getBlendShapeTargetCount(meshIdx));
+        ASSERT_EQ(output->getBlendShapeTargetCount(meshIdx), expected->getBlendShapeTargetCount(meshIdx));
         for (std::uint16_t bsIdx = 0u; bsIdx < expected->getBlendShapeTargetCount(meshIdx); ++bsIdx) {
-            ASSERT_EQ(master->getBlendShapeTargetDeltaCount(meshIdx, bsIdx),
+            ASSERT_EQ(output->getBlendShapeTargetDeltaCount(meshIdx, bsIdx),
                       expected->getBlendShapeTargetDeltaCount(meshIdx, bsIdx));
 
-            auto resultIndices = master->getBlendShapeTargetVertexIndices(meshIdx, bsIdx);
+            auto resultIndices = output->getBlendShapeTargetVertexIndices(meshIdx, bsIdx);
             auto expectedIndices = expected->getBlendShapeTargetVertexIndices(meshIdx, bsIdx);
             ASSERT_EQ(resultIndices, expectedIndices);
 
             for (std::uint32_t deltaIdx = 0u; deltaIdx < expected->getBlendShapeTargetDeltaCount(meshIdx, bsIdx); ++deltaIdx) {
-                Vec3 resultDelta = master->getBlendShapeTargetDelta(meshIdx, bsIdx, deltaIdx);
+                Vec3 resultDelta = output->getBlendShapeTargetDelta(meshIdx, bsIdx, deltaIdx);
                 Vec3 expectedDelta = expected->getBlendShapeTargetDelta(meshIdx, bsIdx, deltaIdx);
                 assertVec3Near(resultDelta, expectedDelta, threshold);
             }
@@ -213,10 +213,10 @@ inline void assertBlendShapePoolDNAIndices(const VariableWidthMatrix<VariableWid
     }
 }
 
-inline void assertNeutralJointTranslations(Reader* master, Reader* expected, float threshold) {
-    ASSERT_EQ(master->getJointCount(), expected->getJointCount());
+inline void assertNeutralJointTranslations(Reader* output, Reader* expected, float threshold) {
+    ASSERT_EQ(output->getJointCount(), expected->getJointCount());
     for (std::uint16_t jointIdx = 0u; jointIdx < expected->getJointCount(); ++jointIdx) {
-        Vec3 resultOffset = master->getNeutralJointTranslation(jointIdx);
+        Vec3 resultOffset = output->getNeutralJointTranslation(jointIdx);
         Vec3 expectedOffset = expected->getNeutralJointTranslation(jointIdx);
         assertVec3Near(resultOffset, expectedOffset, threshold);
     }
@@ -259,10 +259,10 @@ inline void assertNeutralJointPool(const XYZTiledMatrix<16u>& actual, const Matr
     }
 }
 
-inline void assertNeutralJointRotations(Reader* master, Reader* expected, float threshold) {
-    ASSERT_EQ(master->getJointCount(), expected->getJointCount());
+inline void assertNeutralJointRotations(Reader* output, Reader* expected, float threshold) {
+    ASSERT_EQ(output->getJointCount(), expected->getJointCount());
     for (std::uint16_t jointIdx = 0u; jointIdx < expected->getJointCount(); ++jointIdx) {
-        Vec3 resultOffset = master->getNeutralJointRotation(jointIdx);
+        Vec3 resultOffset = output->getNeutralJointRotation(jointIdx);
         Vec3 expectedOffset = expected->getNeutralJointRotation(jointIdx);
         assertVec3Near(resultOffset, expectedOffset, threshold);
     }
@@ -279,18 +279,18 @@ inline void assertNeutralJointRotation(const RawVector3Vector& neutralJoints, Re
     ASSERT_ELEMENTS_AND_SIZE_EQ(neutralJoints.zs, expectedRotationZs);
 }
 
-inline void assertJointBehavior(Reader* master, Reader* expected, float threshold) {
-    ASSERT_EQ(master->getJointGroupCount(), expected->getJointGroupCount());
+inline void assertJointBehavior(Reader* output, Reader* expected, float threshold) {
+    ASSERT_EQ(output->getJointGroupCount(), expected->getJointGroupCount());
     for (std::uint16_t jointGroupIdx = 0u; jointGroupIdx < expected->getJointGroupCount(); ++jointGroupIdx) {
-        auto resultOutputIndices = master->getJointGroupOutputIndices(jointGroupIdx);
+        auto resultOutputIndices = output->getJointGroupOutputIndices(jointGroupIdx);
         auto expectedOutputIndices = expected->getJointGroupOutputIndices(jointGroupIdx);
         ASSERT_EQ(resultOutputIndices, expectedOutputIndices);
 
-        auto resultValues = master->getJointGroupValues(jointGroupIdx);
+        auto resultValues = output->getJointGroupValues(jointGroupIdx);
         auto expectedValues = expected->getJointGroupValues(jointGroupIdx);
         ASSERT_ELEMENTS_NEAR(resultValues, expectedValues, expectedValues.size(), threshold);
 
-        auto resultLODs = master->getJointGroupLODs(jointGroupIdx);
+        auto resultLODs = output->getJointGroupLODs(jointGroupIdx);
         auto expectedLODs = expected->getJointGroupLODs(jointGroupIdx);
         ASSERT_EQ(resultLODs, expectedLODs);
     }
@@ -349,16 +349,16 @@ inline void assertJointBehaviorValues(ConstArrayView<SingleJointBehavior> actual
     }
 }
 
-inline void assertSkinWeights(Reader* master, Reader* expected, float threshold) {
-    ASSERT_EQ(master->getMeshCount(), expected->getMeshCount());
+inline void assertSkinWeights(Reader* output, Reader* expected, float threshold) {
+    ASSERT_EQ(output->getMeshCount(), expected->getMeshCount());
     for (std::uint16_t meshIdx = 0; meshIdx < expected->getMeshCount(); ++meshIdx) {
-        ASSERT_EQ(master->getVertexPositionCount(meshIdx), expected->getVertexPositionCount(meshIdx));
+        ASSERT_EQ(output->getVertexPositionCount(meshIdx), expected->getVertexPositionCount(meshIdx));
         for (std::uint32_t vertexIdx = 0; vertexIdx < expected->getVertexPositionCount(meshIdx); ++vertexIdx) {
-            auto resultWeights = master->getSkinWeightsValues(meshIdx, vertexIdx);
+            auto resultWeights = output->getSkinWeightsValues(meshIdx, vertexIdx);
             auto expectedWeights = expected->getSkinWeightsValues(meshIdx, vertexIdx);
             ASSERT_ELEMENTS_NEAR(resultWeights, expectedWeights, expectedWeights.size(), threshold);
 
-            auto resultJointIndices = master->getSkinWeightsJointIndices(meshIdx, vertexIdx);
+            auto resultJointIndices = output->getSkinWeightsJointIndices(meshIdx, vertexIdx);
             auto expectedJointIndices = expected->getSkinWeightsJointIndices(meshIdx, vertexIdx);
             ASSERT_EQ(resultJointIndices, expectedJointIndices);
         }
diff --git a/Engine/Plugins/Animation/GeneSplicer/Source/GeneSplicerLibTest/Private/gstests/TestGeneSplicer.cpp b/Engine/Plugins/Animation/GeneSplicer/Source/GeneSplicerLibTest/Private/gstests/TestGeneSplicer.cpp
--- a/Engine/Plugins/Animation/GeneSplicer/Source/GeneSplicerLibTest/Private/gstests/TestGeneSplicer.cpp
+++ b/Engine/Plugins/Animation/GeneSplicer/Source/GeneSplicerLibTest/Private/gstests/TestGeneSplicer.cpp
@@ -44,7 +44,7 @@ class TestGeneSplicer : public ::testing::TestWithParam<CalculationType> {
             }
 
             archetype = makeScoped<FixtureReader>(static_cast<std::uint16_t>(FixtureReader::archetype));
-            master = ScopedPtr<GeneSplicerDNAReader>(GeneSplicerDNAReader::create(archetype.get(), &memRes));
+            output = ScopedPtr<GeneSplicerDNAReader>(GeneSplicerDNAReader::create(archetype.get(), &memRes));
             expected = makeScoped<FixtureReader>(static_cast<std::uint16_t>(FixtureReader::expected));
             regionAffiliations.reset(new MockedRegionAffiliationReader{});
             Vector<GenePoolMask> poolMasks {GenePoolMask::NeutralMeshes,
@@ -74,7 +74,7 @@ class TestGeneSplicer : public ::testing::TestWithParam<CalculationType> {
         AlignedMemoryResource memRes;
         Vector<FixtureReaderPtr> smartDNAPtrs;
         Vector<const Reader*> dnaRawPtrs;
-        ScopedPtr<GeneSplicerDNAReader> master;
+        ScopedPtr<GeneSplicerDNAReader> output;
         ScopedPtr<FixtureReader> archetype;
         ScopedPtr<FixtureReader> expected;
         SpliceData spliceData;
@@ -86,29 +86,29 @@ class TestGeneSplicer : public ::testing::TestWithParam<CalculationType> {
 };
 
 TEST_P(TestGeneSplicer, NeutralMeshSplicer) {
-    geneSplicer.spliceNeutralMeshes(&spliceData, master.get());
-    assertNeutralMeshes(master.get(), expected.get(), threshold);
+    geneSplicer.spliceNeutralMeshes(&spliceData, output.get());
+    assertNeutralMeshes(output.get(), expected.get(), threshold);
 }
 
 TEST_P(TestGeneSplicer, BlendShapeSplicer) {
-    geneSplicer.spliceBlendShapes(&spliceData, master.get());
-    assertBlendShapeTargets(master.get(), expected.get(), threshold);
+    geneSplicer.spliceBlendShapes(&spliceData, output.get());
+    assertBlendShapeTargets(output.get(), expected.get(), threshold);
 }
 
 TEST_P(TestGeneSplicer, JointBehaviorSplicer) {
-    geneSplicer.spliceJointBehavior(&spliceData, master.get());
-    assertJointBehavior(master.get(), expected.get(), threshold);
+    geneSplicer.spliceJointBehavior(&spliceData, output.get());
+    assertJointBehavior(output.get(), expected.get(), threshold);
 }
 
 TEST_P(TestGeneSplicer, NeutralJointSplicer) {
-    geneSplicer.spliceNeutralJoints(&spliceData, master.get());
-    assertNeutralJointTranslations(master.get(), expected.get(), threshold);
-    assertNeutralJointRotations(master.get(), expected.get(), threshold);
+    geneSplicer.spliceNeutralJoints(&spliceData, output.get());
+    assertNeutralJointTranslations(output.get(), expected.get(), threshold);
+    assertNeutralJointRotations(output.get(), expected.get(), threshold);
 }
 
 TEST_P(TestGeneSplicer, SkinWeightSplicer) {
-    geneSplicer.spliceSkinWeights(&spliceData, master.get());
-    assertSkinWeights(master.get(), expected.get(), threshold);
+    geneSplicer.spliceSkinWeights(&spliceData, output.get());
+    assertSkinWeights(output.get(), expected.get(), threshold);
 }
 
 TEST_P(TestGeneSplicer, SkinWeightSplicer0Weights) {
@@ -119,8 +119,8 @@ TEST_P(TestGeneSplicer, SkinWeightSplicer0Weights) {
     auto poolName = std::to_string(static_cast<std::uint32_t>(GenePoolMask::SkinWeights));
     spliceData.getPoolParams(poolName.c_str())->setSpliceWeights(0u, spliceWeights.data(),
                                                                  static_cast<std::uint32_t>(spliceWeights.size()));
-    geneSplicer.spliceSkinWeights(&spliceData, master.get());
-    assertSkinWeights(master.get(), archetype.get(), threshold);
+    geneSplicer.spliceSkinWeights(&spliceData, output.get());
+    assertSkinWeights(output.get(), archetype.get(), threshold);
 }
 
 
