Comment out unreachable code
#rnx

diff --git a/Engine/Plugins/Animation/GeneSplicer/Source/GeneSplicerLibTest/Private/GeneSplicerLibTestSuite.cpp b/Engine/Plugins/Animation/GeneSplicer/Source/GeneSplicerLibTest/Private/GeneSplicerLibTestSuite.cpp
--- a/Engine/Plugins/Animation/GeneSplicer/Source/GeneSplicerLibTest/Private/GeneSplicerLibTestSuite.cpp
+++ b/Engine/Plugins/Animation/GeneSplicer/Source/GeneSplicerLibTest/Private/GeneSplicerLibTestSuite.cpp
@@ -61,11 +61,11 @@ bool FGeneSplicerLibTestSuite::RunTest(const FString& Parameters)
 	::testing::GTEST_FLAG(filter) = TestFilter.GetData();
 	return (RUN_ALL_TESTS() == 0);
 
-	TestTrue("GeneSplicerLib Tests", RUN_ALL_TESTS() == 0);
+	/*TestTrue("GeneSplicerLib Tests", RUN_ALL_TESTS() == 0);
 
 	Listeners.Release(&TestPrinter);
 
-	return true;
+	return true;*/
 }
 
 #endif  // WITH_DEV_AUTOMATION_TESTS
\ No newline at end of file
