Remove commented code from GeneSplicerLIbTestSuite.cpp

diff --git a/Engine/Plugins/Animation/GeneSplicer/Source/GeneSplicerLibTest/Private/GeneSplicerLibTestSuite.cpp b/Engine/Plugins/Animation/GeneSplicer/Source/GeneSplicerLibTest/Private/GeneSplicerLibTestSuite.cpp
--- a/Engine/Plugins/Animation/GeneSplicer/Source/GeneSplicerLibTest/Private/GeneSplicerLibTestSuite.cpp
+++ b/Engine/Plugins/Animation/GeneSplicer/Source/GeneSplicerLibTest/Private/GeneSplicerLibTestSuite.cpp
@@ -60,12 +60,6 @@ bool FGeneSplicerLibTestSuite::RunTest(const FString& Parameters)
 	const TArray<ANSICHAR> TestFilter{ TCHAR_TO_ANSI(*Filter), Filter.Len() + 1 };
 	::testing::GTEST_FLAG(filter) = TestFilter.GetData();
 	return (RUN_ALL_TESTS() == 0);
-
-	/*TestTrue("GeneSplicerLib Tests", RUN_ALL_TESTS() == 0);
-
-	Listeners.Release(&TestPrinter);
-
-	return true;*/
 }
 
 #endif  // WITH_DEV_AUTOMATION_TESTS
\ No newline at end of file
