Update RigLogic, DNACalib and GeneSplicer plugin metadata

diff --git a/Engine/Plugins/Animation/DNACalib/DNACalib.uplugin b/Engine/Plugins/Animation/DNACalib/DNACalib.uplugin
--- a/Engine/Plugins/Animation/DNACalib/DNACalib.uplugin
+++ b/Engine/Plugins/Animation/DNACalib/DNACalib.uplugin
@@ -3,15 +3,15 @@
 	"Version": 6.12,
 	"VersionName": "6.12.2",
 	"FriendlyName": "DNACalib Plugin v6.12.2",
-	"Description": "3Lateral DNA Calibration tool plugin",
+	"Description": "DNA Calibration tool plugin",
 	"Category": "Animation",
-	"CreatedBy": "3Lateral",
-	"CreatedByURL": "http://www.3lateral.com/",
+	"CreatedBy": "Epic Games, Inc.",
+	"CreatedByURL": "https://epicgames.com",
 	"DocsURL": "",
 	"MarketplaceURL": "",
 	"SupportURL": "",
 	"CanContainContent": false,
-	"IsBetaVersion": true,
+	"IsBetaVersion": false,
 	"IsExperimentalVersion": false,
 	"Installed": false,
 	"Modules": [
diff --git a/Engine/Plugins/Animation/GeneSplicer/GeneSplicer.uplugin b/Engine/Plugins/Animation/GeneSplicer/GeneSplicer.uplugin
--- a/Engine/Plugins/Animation/GeneSplicer/GeneSplicer.uplugin
+++ b/Engine/Plugins/Animation/GeneSplicer/GeneSplicer.uplugin
@@ -3,15 +3,15 @@
 	"Version": 9.8,
 	"VersionName": "9.8.2",
 	"FriendlyName": "GeneSplicer Plugin v9.8.2",
-	"Description": "3Lateral GeneSplicer plugin for facial animation",
+	"Description": "GeneSplicer plugin for facial animation",
 	"Category": "Animation",
-	"CreatedBy": "3Lateral",
-	"CreatedByURL": "http://www.3lateral.com/",
+	"CreatedBy": "Epic Games, Inc.",
+	"CreatedByURL": "https://epicgames.com",
 	"DocsURL": "",
 	"MarketplaceURL": "",
 	"SupportURL": "",
 	"CanContainContent": false,
-	"IsBetaVersion": true,
+	"IsBetaVersion": false,
 	"Installed": true,
 	"Modules": [
 		{
diff --git a/Engine/Plugins/Animation/RigLogic/RigLogic.uplugin b/Engine/Plugins/Animation/RigLogic/RigLogic.uplugin
--- a/Engine/Plugins/Animation/RigLogic/RigLogic.uplugin
+++ b/Engine/Plugins/Animation/RigLogic/RigLogic.uplugin
@@ -3,15 +3,15 @@
 	"Version": 10.3,
 	"VersionName": "10.2.7",
 	"FriendlyName": "RigLogic Plugin v10.3.0",
-	"Description": "3Lateral RigLogic Plugin for Facial Animation v10.3.0",
+	"Description": "RigLogic Plugin for Facial Animation v10.3.0",
 	"Category": "Animation",
-	"CreatedBy": "3Lateral",
-	"CreatedByURL": "https://www.3lateral.com/",
+	"CreatedBy": "Epic Games, Inc.",
+	"CreatedByURL": "https://epicgames.com",
 	"DocsURL": "",
 	"MarketplaceURL": "",
 	"SupportURL": "",
 	"CanContainContent": true,
-	"IsBetaVersion": true,
+	"IsBetaVersion": false,
 	"Installed": false,
 	"Modules": [
 		{
