Don't use SSE on Windows Arm64

#rnx

diff --git a/Engine/Plugins/Animation/GeneSplicer/Source/GeneSplicerLibTest/GeneSplicerLibTest.Build.cs b/Engine/Plugins/Animation/GeneSplicer/Source/GeneSplicerLibTest/GeneSplicerLibTest.Build.cs
--- a/Engine/Plugins/Animation/GeneSplicer/Source/GeneSplicerLibTest/GeneSplicerLibTest.Build.cs
+++ b/Engine/Plugins/Animation/GeneSplicer/Source/GeneSplicerLibTest/GeneSplicerLibTest.Build.cs
@@ -13,7 +13,10 @@ namespace UnrealBuildTool.Rules
 
 			if (Target.Platform == UnrealTargetPlatform.Win64)
 			{
-				PrivateDefinitions.Add("GS_BUILD_WITH_SSE=1");
+				if (Target.Architecture != UnrealArch.Arm64)
+				{
+					PrivateDefinitions.Add("GS_BUILD_WITH_SSE=1");
+				}
 				PublicDefinitions.Add("GTEST_OS_WINDOWS=1");
 			}
 
