Set GeneSplicer installed flag to false

diff --git a/Engine/Plugins/Animation/GeneSplicer/GeneSplicer.uplugin b/Engine/Plugins/Animation/GeneSplicer/GeneSplicer.uplugin
--- a/Engine/Plugins/Animation/GeneSplicer/GeneSplicer.uplugin
+++ b/Engine/Plugins/Animation/GeneSplicer/GeneSplicer.uplugin
@@ -12,7 +12,7 @@
 	"SupportURL": "",
 	"CanContainContent": false,
 	"IsBetaVersion": false,
-	"Installed": true,
+	"Installed": false,
 	"Modules": [
 		{
 			"Name": "GeneSplicerLib",
