Fix GeneSplicer dependency specification so shipping builds are possible with it

diff --git a/Engine/Plugins/Animation/GeneSplicer/Source/GeneSplicerModule/GeneSplicerModule.Build.cs b/Engine/Plugins/Animation/GeneSplicer/Source/GeneSplicerModule/GeneSplicerModule.Build.cs
--- a/Engine/Plugins/Animation/GeneSplicer/Source/GeneSplicerModule/GeneSplicerModule.Build.cs
+++ b/Engine/Plugins/Animation/GeneSplicer/Source/GeneSplicerModule/GeneSplicerModule.Build.cs
@@ -13,7 +13,6 @@ namespace UnrealBuildTool.Rules
 					"CoreUObject",
 					"Engine",
 					"RigLogicModule",
-					"MessageLog",
 					"Projects"
 				} 
 			);
@@ -21,6 +20,7 @@ namespace UnrealBuildTool.Rules
 			if (Target.Type == TargetType.Editor)
 			{
 				PublicDependencyModuleNames.Add("UnrealEd");
+				PublicDependencyModuleNames.Add("MessageLog");
 			}
 
 			PrivateDependencyModuleNames.AddRange(
@@ -28,7 +28,6 @@ namespace UnrealBuildTool.Rules
 				{
 					"RigLogicModule",
 					"RigLogicLib",
-					"SkeletalMeshUtilitiesCommon",
 					"GeneSplicerLib"
 				}
 			 );
